Copying //UE4/Dev-Editor to Dev-Main (//UE4/Dev-Main) for 4.24

diff --git a/Engine/Source/Runtime/MeshDescription/Public/MeshDescriptionBase.h b/Engine/Source/Runtime/MeshDescription/Public/MeshDescriptionBase.h
new file mode 100644
--- /dev/null
+++ b/Engine/Source/Runtime/MeshDescription/Public/MeshDescriptionBase.h
@@ -0,0 +1,456 @@
+// Copyright 1998-2019 Epic Games, Inc. All Rights Reserved.
+
+#pragma once
+
+#include "MeshDescription.h"
+#include "MeshAttributes.h"
+#include "MeshDescriptionBase.generated.h"
+
+
+
+UCLASS(BlueprintType)
+class MESHDESCRIPTION_API UMeshDescriptionBase : public UObject
+{
+	GENERATED_BODY()
+
+public:
+	// UObject interface
+	virtual void Serialize(FArchive& Ar) override;
+	//
+
+	/** Get a reference to the actual mesh description */
+	FMeshDescription& GetMeshDescription() { return MeshDescription; }
+	const FMeshDescription& GetMeshDescription() const { return MeshDescription; }
+
+	/** Set the mesh description */
+	void SetMeshDescription(FMeshDescription InMeshDescription)
+	{ 
+		MeshDescription = MoveTemp(InMeshDescription);
+	}
+
+	// UMeshDescriptionBase interface
+	virtual void RegisterAttributes();
+	virtual FMeshAttributes& GetRequiredAttributes() { return *RequiredAttributes; }
+	virtual const FMeshAttributes& GetRequiredAttributes() const { return *RequiredAttributes; }
+	//
+
+	/** Reset the contained mesh description */
+	void Reset();
+
+	/** Accessors for mesh element arrays */
+	FVertexArray& Vertices() { return MeshDescription.Vertices(); }
+	const FVertexArray& Vertices() const { return MeshDescription.Vertices(); }
+
+	FVertexInstanceArray& VertexInstances() { return MeshDescription.VertexInstances(); }
+	const FVertexInstanceArray& VertexInstances() const { return MeshDescription.VertexInstances(); }
+
+	FEdgeArray& Edges() { return MeshDescription.Edges(); }
+	const FEdgeArray& Edges() const { return MeshDescription.Edges(); }
+
+	FTriangleArray& Triangles() { return MeshDescription.Triangles(); }
+	const FTriangleArray& Triangles() const { return MeshDescription.Triangles(); }
+
+	FPolygonArray& Polygons() { return MeshDescription.Polygons(); }
+	const FPolygonArray& Polygons() const { return MeshDescription.Polygons(); }
+
+	FPolygonGroupArray& PolygonGroups() { return MeshDescription.PolygonGroups(); }
+	const FPolygonGroupArray& PolygonGroups() const { return MeshDescription.PolygonGroups(); }
+
+	/** Accessors for mesh element attributes */
+	TAttributesSet<FVertexID>& VertexAttributes() { return MeshDescription.VertexAttributes(); }
+	const TAttributesSet<FVertexID>& VertexAttributes() const { return MeshDescription.VertexAttributes(); }
+
+	TAttributesSet<FVertexInstanceID>& VertexInstanceAttributes() { return MeshDescription.VertexInstanceAttributes(); }
+	const TAttributesSet<FVertexInstanceID>& VertexInstanceAttributes() const { return MeshDescription.VertexInstanceAttributes(); }
+
+	TAttributesSet<FEdgeID>& EdgeAttributes() { return MeshDescription.EdgeAttributes(); }
+	const TAttributesSet<FEdgeID>& EdgeAttributes() const { return MeshDescription.EdgeAttributes(); }
+
+	TAttributesSet<FTriangleID>& TriangleAttributes() { return MeshDescription.TriangleAttributes(); }
+	const TAttributesSet<FTriangleID>& TriangleAttributes() const { return MeshDescription.TriangleAttributes(); }
+
+	TAttributesSet<FPolygonID>& PolygonAttributes() { return MeshDescription.PolygonAttributes(); }
+	const TAttributesSet<FPolygonID>& PolygonAttributes() const { return MeshDescription.PolygonAttributes(); }
+
+	TAttributesSet<FPolygonGroupID>& PolygonGroupAttributes() { return MeshDescription.PolygonGroupAttributes(); }
+	const TAttributesSet<FPolygonGroupID>& PolygonGroupAttributes() const { return MeshDescription.PolygonGroupAttributes(); }
+
+	/** Accessors for cached vertex position array */
+	TVertexAttributesRef<FVector> GetVertexPositions() { return GetRequiredAttributes().GetVertexPositions(); }
+	TVertexAttributesConstRef<FVector> GetVertexPositions() const { return GetRequiredAttributes().GetVertexPositions(); }
+
+public:
+
+	/** Empty the mesh description */
+	UFUNCTION(BlueprintCallable, Category="MeshDescription")
+	void Empty();
+
+	/** Return whether the mesh description is empty */
+	UFUNCTION(BlueprintPure, Category="MeshDescription")
+	bool IsEmpty() const;
+
+
+	///////////////////////////////////////////////////////
+	// Create / remove mesh elements
+
+	/** Reserves space for this number of new vertices */
+	UFUNCTION(BlueprintCallable, Category="MeshDescription")
+	void ReserveNewVertices(int32 NumberOfNewVertices);
+
+	/** Adds a new vertex to the mesh and returns its ID */
+	UFUNCTION(BlueprintCallable, Category="MeshDescription")
+	FVertexID CreateVertex();
+
+	/** Adds a new vertex to the mesh with the given ID */
+	UFUNCTION(BlueprintCallable, Category="MeshDescription")
+	void CreateVertexWithID(FVertexID VertexID);
+
+	/** Deletes a vertex from the mesh */
+	UFUNCTION(BlueprintCallable, Category="MeshDescription")
+	void DeleteVertex(FVertexID VertexID);
+
+	/** Returns whether the passed vertex ID is valid */
+	UFUNCTION(BlueprintPure, Category="MeshDescription")
+	bool IsVertexValid(FVertexID VertexID) const;
+
+	/** Reserves space for this number of new vertex instances */
+	UFUNCTION(BlueprintCallable, Category="MeshDescription")
+	void ReserveNewVertexInstances(int32 NumberOfNewVertexInstances);
+
+	/** Adds a new vertex instance to the mesh and returns its ID */
+	UFUNCTION(BlueprintCallable, Category="MeshDescription")
+	FVertexInstanceID CreateVertexInstance(FVertexID VertexID);
+
+	/** Adds a new vertex instance to the mesh with the given ID */
+	UFUNCTION(BlueprintCallable, Category="MeshDescription")
+	void CreateVertexInstanceWithID(FVertexInstanceID VertexInstanceID, FVertexID VertexID);
+
+	/** Deletes a vertex instance from a mesh */
+	UFUNCTION(BlueprintCallable, Category="MeshDescription")
+	void DeleteVertexInstance(FVertexInstanceID VertexInstanceID, TArray<FVertexID>& OrphanedVertices);
+
+	/** Returns whether the passed vertex instance ID is valid */
+	UFUNCTION(BlueprintPure, Category="MeshDescription")
+	bool IsVertexInstanceValid(FVertexInstanceID VertexInstanceID) const;
+
+	/** Reserves space for this number of new edges */
+	UFUNCTION(BlueprintCallable, Category="MeshDescription")
+	void ReserveNewEdges(int32 NumberOfNewEdges);
+
+	/** Adds a new edge to the mesh and returns its ID */
+	UFUNCTION(BlueprintCallable, Category="MeshDescription")
+	FEdgeID CreateEdge(FVertexID VertexID0, FVertexID VertexID1);
+
+	/** Adds a new edge to the mesh with the given ID */
+	UFUNCTION(BlueprintCallable, Category="MeshDescription")
+	void CreateEdgeWithID(FEdgeID EdgeID, FVertexID VertexID0, FVertexID VertexID1);
+
+	/** Deletes an edge from a mesh */
+	UFUNCTION(BlueprintCallable, Category="MeshDescription")
+	void DeleteEdge(FEdgeID EdgeID, TArray<FVertexID>& OrphanedVertices);
+
+	/** Returns whether the passed edge ID is valid */
+	UFUNCTION(BlueprintPure, Category="MeshDescription")
+	bool IsEdgeValid(FEdgeID EdgeID) const;
+
+	/** Reserves space for this number of new triangles */
+	UFUNCTION(BlueprintCallable, Category="MeshDescription")
+	void ReserveNewTriangles(int32 NumberOfNewTriangles);
+
+	/** Adds a new triangle to the mesh and returns its ID. This will also make an encapsulating polygon, and any missing edges. */
+	UFUNCTION(BlueprintCallable, Category="MeshDescription")
+	FTriangleID CreateTriangle(FPolygonGroupID PolygonGroupID, const TArray<FVertexInstanceID>& VertexInstanceIDs, TArray<FEdgeID>& NewEdgeIDs);
+
+	/** Adds a new triangle to the mesh with the given ID. This will also make an encapsulating polygon, and any missing edges. */
+	UFUNCTION(BlueprintCallable, Category="MeshDescription")
+	void CreateTriangleWithID(FTriangleID TriangleID, FPolygonGroupID PolygonGroupID, const TArray<FVertexInstanceID>& VertexInstanceIDs, TArray<FEdgeID>& NewEdgeIDs);
+
+	/** Deletes a triangle from the mesh */
+	UFUNCTION(BlueprintCallable, Category="MeshDescription")
+	void DeleteTriangle(FTriangleID TriangleID, TArray<FEdgeID>& OrphanedEdges, TArray<FVertexInstanceID>& OrphanedVertexInstances, TArray<FPolygonGroupID>& OrphanedPolygonGroupsPtr);
+
+	/** Returns whether the passed triangle ID is valid */
+	UFUNCTION(BlueprintPure, Category="MeshDescription")
+	bool IsTriangleValid(const FTriangleID TriangleID) const;
+
+	/** Reserves space for this number of new polygons */
+	UFUNCTION(BlueprintCallable, Category="MeshDescription")
+	void ReserveNewPolygons(const int32 NumberOfNewPolygons);
+
+	/** Adds a new polygon to the mesh and returns its ID. This will also make any missing edges, and all constituent triangles. */
+	UFUNCTION(BlueprintCallable, Category="MeshDescription")
+	FPolygonID CreatePolygon(FPolygonGroupID PolygonGroupID, TArray<FVertexInstanceID>& VertexInstanceIDs, TArray<FEdgeID>& NewEdgeIDs);
+
+	/** Adds a new polygon to the mesh with the given ID. This will also make any missing edges, and all constituent triangles. */
+	UFUNCTION(BlueprintCallable, Category="MeshDescription")
+	void CreatePolygonWithID(FPolygonID PolygonID, FPolygonGroupID PolygonGroupID, TArray<FVertexInstanceID>& VertexInstanceIDs, TArray<FEdgeID>& NewEdgeIDs);
+
+	/** Deletes a polygon from the mesh */
+	UFUNCTION(BlueprintCallable, Category="MeshDescription")
+	void DeletePolygon(FPolygonID PolygonID, TArray<FEdgeID>& OrphanedEdges, TArray<FVertexInstanceID>& OrphanedVertexInstances, TArray<FPolygonGroupID>& OrphanedPolygonGroups);
+
+	/** Returns whether the passed polygon ID is valid */
+	UFUNCTION(BlueprintPure, Category="MeshDescription")
+	bool IsPolygonValid(FPolygonID PolygonID) const;
+
+	/** Reserves space for this number of new polygon groups */
+	UFUNCTION(BlueprintCallable, Category="MeshDescription")
+	void ReserveNewPolygonGroups(int32 NumberOfNewPolygonGroups);
+
+	/** Adds a new polygon group to the mesh and returns its ID */
+	UFUNCTION(BlueprintCallable, Category="MeshDescription")
+	FPolygonGroupID CreatePolygonGroup();
+
+	/** Adds a new polygon group to the mesh with the given ID */
+	UFUNCTION(BlueprintCallable, Category="MeshDescription")
+	void CreatePolygonGroupWithID(FPolygonGroupID PolygonGroupID);
+
+	/** Deletes a polygon group from the mesh */
+	UFUNCTION(BlueprintCallable, Category="MeshDescription")
+	void DeletePolygonGroup(FPolygonGroupID PolygonGroupID);
+
+	/** Returns whether the passed polygon group ID is valid */
+	UFUNCTION(BlueprintPure, Category="MeshDescription")
+	bool IsPolygonGroupValid(FPolygonGroupID PolygonGroupID) const;
+
+
+	//////////////////////////////////////////////////////////////////////
+	// Vertex operations
+
+	/** Returns whether a given vertex is orphaned, i.e. it doesn't form part of any polygon */
+	UFUNCTION(BlueprintPure, Category="MeshDescription")
+	bool IsVertexOrphaned(FVertexID VertexID) const;
+
+	/** Returns the edge ID defined by the two given vertex IDs, if there is one; otherwise FEdgeID::Invalid */
+	UFUNCTION(BlueprintPure, Category="MeshDescription")
+	FEdgeID GetVertexPairEdge(FVertexID VertexID0, FVertexID VertexID1) const;
+
+	/** Returns reference to an array of Edge IDs connected to this vertex */
+	UFUNCTION(BlueprintPure, Category="MeshDescription")
+	void GetVertexConnectedEdges(FVertexID VertexID, TArray<FEdgeID>& OutEdgeIDs) const;
+
+	/** Returns number of edges connected to this vertex */
+	UFUNCTION(BlueprintPure, Category="MeshDescription")
+	int32 GetNumVertexConnectedEdges(FVertexID VertexID) const;
+
+	/** Returns reference to an array of VertexInstance IDs instanced from this vertex */
+	UFUNCTION(BlueprintPure, Category="MeshDescription")
+	void GetVertexVertexInstances(FVertexID VertexID, TArray<FVertexInstanceID>& OutVertexInstanceIDs) const;
+
+	/** Returns number of vertex instances created from this vertex */
+	UFUNCTION(BlueprintPure, Category="MeshDescription")
+	int32 GetNumVertexVertexInstances(FVertexID VertexID) const;
+
+	/** Returns the triangles connected to this vertex */
+	UFUNCTION(BlueprintPure, Category="MeshDescription")
+	void GetVertexConnectedTriangles(FVertexID VertexID, TArray<FTriangleID>& OutConnectedTriangleIDs) const;
+
+	/** Returns number of triangles connected to this vertex */
+	UFUNCTION(BlueprintPure, Category="MeshDescription")
+	int32 GetNumVertexConnectedTriangles(FVertexID VertexID) const;
+
+	/** Returns the polygons connected to this vertex */
+	UFUNCTION(BlueprintPure, Category="MeshDescription")
+	void GetVertexConnectedPolygons(FVertexID VertexID, TArray<FPolygonID>& OutConnectedPolygonIDs) const;
+
+	/** Returns the number of polygons connected to this vertex */
+	UFUNCTION(BlueprintPure, Category="MeshDescription")
+	int32 GetNumVertexConnectedPolygons(FVertexID VertexID) const;
+
+	/** Returns the vertices adjacent to this vertex */
+	UFUNCTION(BlueprintPure, Category="MeshDescription")
+	void GetVertexAdjacentVertices(FVertexID VertexID, TArray<FVertexID>& OutAdjacentVertexIDs) const;
+
+	/** Gets a vertex position */
+	UFUNCTION(BlueprintPure, Category="MeshDescription")
+	FVector GetVertexPosition(FVertexID VertexID) const;
+
+	/** Sets a vertex position */
+	UFUNCTION(BlueprintCallable, Category="MeshDescription")
+	void SetVertexPosition(FVertexID VertexID, const FVector& Position);
+
+
+	//////////////////////////////////////////////////////////////////////
+	// Vertex instance operations
+
+	/** Returns the vertex ID associated with the given vertex instance */
+	UFUNCTION(BlueprintPure, Category="MeshDescription")
+	FVertexID GetVertexInstanceVertex(FVertexInstanceID VertexInstanceID) const;
+
+	/** Returns the edge ID defined by the two given vertex instance IDs, if there is one; otherwise FEdgeID::Invalid */
+	UFUNCTION(BlueprintPure, Category="MeshDescription")
+	FEdgeID GetVertexInstancePairEdge(FVertexInstanceID VertexInstanceID0, FVertexInstanceID VertexInstanceID1) const;
+
+	/** Returns reference to an array of Triangle IDs connected to this vertex instance */
+	UFUNCTION(BlueprintPure, Category="MeshDescription")
+	void GetVertexInstanceConnectedTriangles(FVertexInstanceID VertexInstanceID, TArray<FTriangleID>& OutConnectedTriangleIDs) const;
+
+	/** Returns the number of triangles connected to this vertex instance */
+	UFUNCTION(BlueprintPure, Category="MeshDescription")
+	int32 GetNumVertexInstanceConnectedTriangles(FVertexInstanceID VertexInstanceID) const;
+
+	/** Returns the polygons connected to this vertex instance */
+	UFUNCTION(BlueprintPure, Category="MeshDescription")
+	void GetVertexInstanceConnectedPolygons(FVertexInstanceID VertexInstanceID, TArray<FPolygonID>& OutConnectedPolygonIDs) const;
+
+	/** Returns the number of polygons connected to this vertex instance. */
+	UFUNCTION(BlueprintPure, Category="MeshDescription")
+	int32 GetNumVertexInstanceConnectedPolygons(FVertexInstanceID VertexInstanceID) const;
+
+
+	//////////////////////////////////////////////////////////////////////
+	// Edge operations
+
+	/** Determine whether a given edge is an internal edge between triangles of a polygon */
+	UFUNCTION(BlueprintPure, Category="MeshDescription")
+	bool IsEdgeInternal(FEdgeID EdgeID) const;
+
+	/** Determine whether a given edge is an internal edge between triangles of a specific polygon */
+	UFUNCTION(BlueprintPure, Category="MeshDescription")
+	bool IsEdgeInternalToPolygon(FEdgeID EdgeID, FPolygonID PolygonID) const;
+
+	/** Returns reference to an array of triangle IDs connected to this edge */
+	UFUNCTION(BlueprintPure, Category="MeshDescription")
+	void GetEdgeConnectedTriangles(FEdgeID EdgeID, TArray<FTriangleID>& OutConnectedTriangleIDs) const;
+
+	/** Returns the number of triangles connected to this edge */
+	UFUNCTION(BlueprintPure, Category="MeshDescription")
+	int32 GetNumEdgeConnectedTriangles(FEdgeID EdgeID) const;
+
+	/** Returns the polygons connected to this edge */
+	UFUNCTION(BlueprintPure, Category="MeshDescription")
+	void GetEdgeConnectedPolygons(FEdgeID EdgeID, TArray<FPolygonID>& OutConnectedPolygonIDs) const;
+
+	/** Returns the number of polygons connected to this edge */
+	UFUNCTION(BlueprintPure, Category="MeshDescription")
+	int32 GetNumEdgeConnectedPolygons(FEdgeID EdgeID) const;
+
+	/** Returns the vertex ID corresponding to one of the edge endpoints */
+	UFUNCTION(BlueprintPure, Category="MeshDescription")
+	FVertexID GetEdgeVertex(FEdgeID EdgeID, int32 VertexNumber) const;
+
+	/** Returns a pair of vertex IDs defining the edge */
+	UFUNCTION(BlueprintPure, Category="MeshDescription")
+	void GetEdgeVertices(const FEdgeID EdgeID, TArray<FVertexID>& OutVertexIDs) const;
+
+
+	//////////////////////////////////////////////////////////////////////
+	// Triangle operations
+
+	/** Get the polygon which contains this triangle */
+	UFUNCTION(BlueprintPure, Category="MeshDescription")
+	FPolygonID GetTrianglePolygon(FTriangleID TriangleID) const;
+
+	/** Get the polygon group which contains this triangle */
+	UFUNCTION(BlueprintPure, Category="MeshDescription")
+	FPolygonGroupID GetTrianglePolygonGroup(FTriangleID TriangleID) const;
+
+	/** Determines if this triangle is part of an n-gon */
+	UFUNCTION(BlueprintPure, Category="MeshDescription")
+	bool IsTrianglePartOfNgon(FTriangleID TriangleID) const;
+
+	/** Get the vertex instances which define this triangle */
+	UFUNCTION(BlueprintPure, Category="MeshDescription")
+	void GetTriangleVertexInstances(FTriangleID TriangleID, TArray<FVertexInstanceID>& OutVertexInstanceIDs) const;
+
+	/** Get the specified vertex instance by index */
+	UFUNCTION(BlueprintPure, Category="MeshDescription")
+	FVertexInstanceID GetTriangleVertexInstance(FTriangleID TriangleID, int32 Index) const;
+
+	/** Returns the vertices which define this triangle */
+	UFUNCTION(BlueprintPure, Category="MeshDescription")
+	void GetTriangleVertices(FTriangleID TriangleID, TArray<FVertexID>& OutVertexIDs) const;
+
+	/** Returns the edges which define this triangle */
+	UFUNCTION(BlueprintPure, Category="MeshDescription")
+	void GetTriangleEdges(FTriangleID TriangleID, TArray<FEdgeID>& OutEdgeIDs) const;
+
+	/** Returns the adjacent triangles to this triangle */
+	UFUNCTION(BlueprintPure, Category="MeshDescription")
+	void GetTriangleAdjacentTriangles(FTriangleID TriangleID, TArray<FTriangleID>& OutTriangleIDs) const;
+
+	/** Return the vertex instance which corresponds to the given vertex on the given triangle, or FVertexInstanceID::Invalid */
+	UFUNCTION(BlueprintPure, Category="MeshDescription")
+	FVertexInstanceID GetVertexInstanceForTriangleVertex(FTriangleID TriangleID, FVertexID VertexID) const;
+
+
+	//////////////////////////////////////////////////////////////////////
+	// Polygon operations
+
+	/** Return reference to an array of triangle IDs which comprise this polygon */
+	UFUNCTION(BlueprintPure, Category="MeshDescription")
+	void GetPolygonTriangles(FPolygonID PolygonID, TArray<FTriangleID>& OutTriangleIDs) const;
+
+	/** Return the number of triangles which comprise this polygon */
+	UFUNCTION(BlueprintPure, Category="MeshDescription")
+	int32 GetNumPolygonTriangles(FPolygonID PolygonID) const;
+
+	/** Returns reference to an array of VertexInstance IDs forming the perimeter of this polygon */
+	UFUNCTION(BlueprintPure, Category="MeshDescription")
+	void GetPolygonVertexInstances(FPolygonID PolygonID, TArray<FVertexInstanceID>& OutVertexInstanceIDs) const;
+
+	/** Returns the number of vertices this polygon has */
+	UFUNCTION(BlueprintPure, Category="MeshDescription")
+	int32 GetNumPolygonVertices(FPolygonID PolygonID) const;
+
+	/** Returns the vertices which form the polygon perimeter */
+	UFUNCTION(BlueprintPure, Category="MeshDescription")
+	void GetPolygonVertices(FPolygonID PolygonID, TArray<FVertexID>& OutVertexIDs) const;
+
+	/** Returns the edges which form the polygon perimeter */
+	UFUNCTION(BlueprintPure, Category="MeshDescription")
+	void GetPolygonPerimeterEdges(FPolygonID PolygonID, TArray<FEdgeID>& OutEdgeIDs) const;
+
+	/** Populate the provided array with a list of edges which are internal to the polygon, i.e. those which separate
+	    constituent triangles. */
+	UFUNCTION(BlueprintPure, Category="MeshDescription")
+	void GetPolygonInternalEdges(FPolygonID PolygonID, TArray<FEdgeID>& OutEdgeIDs) const;
+
+	/** Return the number of internal edges in this polygon */
+	UFUNCTION(BlueprintPure, Category="MeshDescription")
+	int32 GetNumPolygonInternalEdges(FPolygonID PolygonID) const;
+
+	/** Populates the passed array with adjacent polygons */
+	UFUNCTION(BlueprintPure, Category="MeshDescription")
+	void GetPolygonAdjacentPolygons(FPolygonID PolygonID, TArray<FPolygonID>& OutPolygonIDs) const;
+
+	/** Return the polygon group associated with a polygon */
+	UFUNCTION(BlueprintPure, Category="MeshDescription")
+	FPolygonGroupID GetPolygonPolygonGroup(FPolygonID PolygonID) const;
+
+	/** Return the vertex instance which corresponds to the given vertex on the given polygon, or FVertexInstanceID::Invalid */
+	UFUNCTION(BlueprintPure, Category="MeshDescription")
+	FVertexInstanceID GetVertexInstanceForPolygonVertex(FPolygonID PolygonID, FVertexID VertexID) const;
+
+	/** Set the vertex instance at the given index around the polygon to the new value */
+	UFUNCTION(BlueprintCallable, Category="MeshDescription")
+	void SetPolygonVertexInstance(FPolygonID PolygonID, int32 PerimeterIndex, FVertexInstanceID VertexInstanceID);
+
+	/** Sets the polygon group associated with a polygon */
+	UFUNCTION(BlueprintCallable, Category="MeshDescription")
+	void SetPolygonPolygonGroup(FPolygonID PolygonID, FPolygonGroupID PolygonGroupID);
+
+	/** Reverse the winding order of the vertices of this polygon */
+	UFUNCTION(BlueprintCallable, Category="MeshDescription")
+	void ReversePolygonFacing(FPolygonID PolygonID);
+
+	/** Generates triangles and internal edges for the given polygon */
+	UFUNCTION(BlueprintCallable, Category="MeshDescription")
+	void ComputePolygonTriangulation(FPolygonID PolygonID);
+
+
+	//////////////////////////////////////////////////////////////////////
+	// Polygon group operations
+
+	/** Returns the polygons associated with the given polygon group */
+	UFUNCTION(BlueprintPure, Category="MeshDescription")
+	void GetPolygonGroupPolygons(FPolygonGroupID PolygonGroupID, TArray<FPolygonID>& OutPolygonIDs) const;
+
+	/** Returns the number of polygons in this polygon group */
+	UFUNCTION(BlueprintPure, Category="MeshDescription")
+	int32 GetNumPolygonGroupPolygons(FPolygonGroupID PolygonGroupID) const;
+
+protected:
+	FMeshDescription MeshDescription;
+	TUniquePtr<FMeshAttributes> RequiredAttributes;
+};
diff --git a/Engine/Source/Runtime/MovieSceneTracks/Public/Sections/MovieScene3DTransformSection.h b/Engine/Source/Runtime/MovieSceneTracks/Public/Sections/MovieScene3DTransformSection.h
--- a/Engine/Source/Runtime/MovieSceneTracks/Public/Sections/MovieScene3DTransformSection.h
+++ b/Engine/Source/Runtime/MovieSceneTracks/Public/Sections/MovieScene3DTransformSection.h
@@ -9,6 +9,7 @@
 #include "MovieSceneSection.h"
 #include "MovieSceneKeyStruct.h"
 #include "Channels/MovieSceneFloatChannel.h"
+#include "TransformData.h"
 #include "MovieScene3DTransformSection.generated.h"
 
 
@@ -24,36 +25,6 @@ enum class EShow3DTrajectory : uint8
 #endif
 
 
-/**
-* Stores information about a transform for the purpose of adding keys to a transform section
-*/
-struct FTransformData
-{
-	/** Translation component */
-	FVector Translation;
-	/** Rotation component */
-	FRotator Rotation;
-	/** Scale component */
-	FVector Scale;
-
-	FTransformData()
-		: Translation(ForceInitToZero)
-		, Rotation(ForceInitToZero)
-		, Scale(ForceInitToZero)
-	{}
-
-	/**
-	* Constructor.  Builds the data from a scene component
-	* Uses relative transform only
-	*
-	* @param InComponent	The component to build from
-	*/
-	FTransformData(const USceneComponent* InComponent)
-		: Translation(InComponent->GetRelativeLocation())
-		, Rotation(InComponent->GetRelativeRotation())
-		, Scale(InComponent->GetRelativeScale3D())
-	{}
-};
 
 /**
  * Proxy structure for translation keys in 3D transform sections.
diff --git a/Engine/Source/Runtime/StaticMeshDescription/Private/StaticMeshAttributes.cpp b/Engine/Source/Runtime/StaticMeshDescription/Private/StaticMeshAttributes.cpp
new file mode 100644
--- /dev/null
+++ b/Engine/Source/Runtime/StaticMeshDescription/Private/StaticMeshAttributes.cpp
@@ -0,0 +1,54 @@
+// Copyright 1998-2019 Epic Games, Inc. All Rights Reserved.
+
+#include "StaticMeshAttributes.h"
+
+namespace MeshAttribute
+{
+	const FName Vertex::CornerSharpness("CornerSharpness");
+
+	const FName VertexInstance::TextureCoordinate("TextureCoordinate");
+	const FName VertexInstance::Normal("Normal");
+	const FName VertexInstance::Tangent("Tangent");
+	const FName VertexInstance::BinormalSign("BinormalSign");
+	const FName VertexInstance::Color("Color");
+
+	const FName Edge::IsHard("IsHard");
+	const FName Edge::IsUVSeam("IsUVSeam");
+	const FName Edge::CreaseSharpness("CreaseSharpness");
+
+	const FName Polygon::Normal("Normal");
+	const FName Polygon::Tangent("Tangent");
+	const FName Polygon::Binormal("Binormal");
+	const FName Polygon::Center("Center");
+
+	const FName PolygonGroup::ImportedMaterialSlotName("ImportedMaterialSlotName");
+	const FName PolygonGroup::EnableCollision("EnableCollision");
+	const FName PolygonGroup::CastShadow("CastShadow");
+}
+
+
+
+void FStaticMeshAttributes::Register()
+{
+	// Add basic vertex attributes
+	MeshDescription.VertexAttributes().RegisterAttribute<float>(MeshAttribute::Vertex::CornerSharpness, 1, 0.0f, EMeshAttributeFlags::Lerpable);
+
+	// Add basic vertex instance attributes
+	MeshDescription.VertexInstanceAttributes().RegisterAttribute<FVector2D>(MeshAttribute::VertexInstance::TextureCoordinate, 1, FVector2D::ZeroVector, EMeshAttributeFlags::Lerpable);
+	MeshDescription.VertexInstanceAttributes().RegisterAttribute<FVector>(MeshAttribute::VertexInstance::Normal, 1, FVector::ZeroVector, EMeshAttributeFlags::AutoGenerated);
+	MeshDescription.VertexInstanceAttributes().RegisterAttribute<FVector>(MeshAttribute::VertexInstance::Tangent, 1, FVector::ZeroVector, EMeshAttributeFlags::AutoGenerated);
+	MeshDescription.VertexInstanceAttributes().RegisterAttribute<float>(MeshAttribute::VertexInstance::BinormalSign, 1, 0.0f, EMeshAttributeFlags::AutoGenerated);
+	MeshDescription.VertexInstanceAttributes().RegisterAttribute<FVector4>(MeshAttribute::VertexInstance::Color, 1, FVector4(1.0f, 1.0f, 1.0f, 1.0f), EMeshAttributeFlags::Lerpable);
+
+	// Add basic edge attributes
+	MeshDescription.EdgeAttributes().RegisterAttribute<bool>(MeshAttribute::Edge::IsHard, 1, false);
+	MeshDescription.EdgeAttributes().RegisterAttribute<float>(MeshAttribute::Edge::CreaseSharpness, 1, 0.0f, EMeshAttributeFlags::Lerpable);
+
+	// Add basic polygon attributes
+
+	// Add basic polygon group attributes
+	MeshDescription.PolygonGroupAttributes().RegisterAttribute<FName>(MeshAttribute::PolygonGroup::ImportedMaterialSlotName); //The unique key to match the mesh material slot
+
+	// Call super class
+	FMeshAttributes::Register();
+}
diff --git a/Engine/Source/Runtime/StaticMeshDescription/Private/StaticMeshDescription.cpp b/Engine/Source/Runtime/StaticMeshDescription/Private/StaticMeshDescription.cpp
new file mode 100644
--- /dev/null
+++ b/Engine/Source/Runtime/StaticMeshDescription/Private/StaticMeshDescription.cpp
@@ -0,0 +1,133 @@
+// Copyright 1998-2019 Epic Games, Inc. All Rights Reserved.
+
+#include "StaticMeshDescription.h"
+#include "StaticMeshAttributes.h"
+#include "StaticMeshOperations.h"
+
+
+void UStaticMeshDescription::RegisterAttributes()
+{
+	RequiredAttributes = MakeUnique<FStaticMeshAttributes>(MeshDescription);
+	RequiredAttributes->Register();
+}
+
+
+FVector2D UStaticMeshDescription::GetVertexInstanceUV(FVertexInstanceID VertexInstanceID, int32 UVIndex) const
+{
+	if (!MeshDescription.IsVertexInstanceValid(VertexInstanceID))
+	{
+		UE_LOG(LogMeshDescription, Warning, TEXT("GetVertexInstanceUV: VertexInstanceID %d doesn't exist."), VertexInstanceID.GetValue());
+		return FVector2D::ZeroVector;
+	}
+
+	if (!MeshDescription.VertexInstanceAttributes().HasAttribute(MeshAttribute::VertexInstance::TextureCoordinate))
+	{
+		UE_LOG(LogMeshDescription, Warning, TEXT("GetVertexInstanceUV: VertexInstanceAttribute TextureCoordinate doesn't exist."));
+		return FVector2D::ZeroVector;
+	}
+
+	return MeshDescription.VertexInstanceAttributes().GetAttribute<FVector2D>(VertexInstanceID, MeshAttribute::VertexInstance::TextureCoordinate, UVIndex);
+}
+
+
+void UStaticMeshDescription::SetVertexInstanceUV(FVertexInstanceID VertexInstanceID, FVector2D UV, int32 UVIndex)
+{
+	if (!MeshDescription.IsVertexInstanceValid(VertexInstanceID))
+	{
+		UE_LOG(LogMeshDescription, Warning, TEXT("SetVertexInstanceUV: VertexInstanceID %d doesn't exist."), VertexInstanceID.GetValue());
+		return;
+	}
+
+	if (!MeshDescription.VertexInstanceAttributes().HasAttribute(MeshAttribute::VertexInstance::TextureCoordinate))
+	{
+		UE_LOG(LogMeshDescription, Warning, TEXT("SetVertexInstanceUV: VertexInstanceAttribute TextureCoordinate doesn't exist."));
+		return;
+	}
+
+	MeshDescription.VertexInstanceAttributes().SetAttribute(VertexInstanceID, MeshAttribute::VertexInstance::TextureCoordinate, UVIndex, UV);
+}
+
+
+void UStaticMeshDescription::SetPolygonGroupMaterialSlotName(FPolygonGroupID PolygonGroupID, const FName& SlotName)
+{
+	if (!MeshDescription.IsPolygonGroupValid(PolygonGroupID))
+	{
+		UE_LOG(LogMeshDescription, Warning, TEXT("SetPolygonGroupMaterialSlotName: PolygonGroupID %d doesn't exist."), PolygonGroupID.GetValue());
+		return;
+	}
+
+	if (!MeshDescription.PolygonGroupAttributes().HasAttribute(MeshAttribute::PolygonGroup::ImportedMaterialSlotName))
+	{
+		UE_LOG(LogMeshDescription, Warning, TEXT("SetPolygonGroupMaterialSlotName: PolygonGroupAttribute ImportedMaterialSlotName doesn't exist."));
+		return;
+	}
+
+	MeshDescription.PolygonGroupAttributes().SetAttribute(PolygonGroupID, MeshAttribute::PolygonGroup::ImportedMaterialSlotName, 0, SlotName);
+}
+
+
+void UStaticMeshDescription::CreateCube(FVector Center, FVector HalfExtents, FPolygonGroupID PolygonGroup,
+										FPolygonID& PolygonID_PlusX,
+										FPolygonID& PolygonID_MinusX,
+										FPolygonID& PolygonID_PlusY, 
+										FPolygonID& PolygonID_MinusY,
+										FPolygonID& PolygonID_PlusZ,
+										FPolygonID& PolygonID_MinusZ)
+{
+	TVertexAttributesRef<FVector> Positions = GetVertexPositions();
+
+	FVertexID VertexIDs[8];
+
+	MeshDescription.ReserveNewVertices(8);
+	for (int32 Index = 0; Index < 8; ++Index)
+	{
+		VertexIDs[Index] = MeshDescription.CreateVertex();
+	}
+
+	Positions[VertexIDs[0]] = Center + HalfExtents * FVector( 1.0f, -1.0f,  1.0f);
+	Positions[VertexIDs[1]] = Center + HalfExtents * FVector( 1.0f,  1.0f,  1.0f);
+	Positions[VertexIDs[2]] = Center + HalfExtents * FVector(-1.0f,  1.0f,  1.0f);
+	Positions[VertexIDs[3]] = Center + HalfExtents * FVector(-1.0f, -1.0f,  1.0f);
+	Positions[VertexIDs[4]] = Center + HalfExtents * FVector(-1.0f,  1.0f, -1.0f);
+	Positions[VertexIDs[5]] = Center + HalfExtents * FVector(-1.0f, -1.0f, -1.0f);
+	Positions[VertexIDs[6]] = Center + HalfExtents * FVector( 1.0f, -1.0f, -1.0f);
+	Positions[VertexIDs[7]] = Center + HalfExtents * FVector( 1.0f,  1.0f, -1.0f);
+
+	auto MakePolygon = [this, &VertexIDs, PolygonGroup](int32 P0, int32 P1, int32 P2, int32 P3) -> FPolygonID
+	{
+		FVertexInstanceID VertexInstanceIDs[4];
+		VertexInstanceIDs[0] = MeshDescription.CreateVertexInstance(VertexIDs[P0]);
+		VertexInstanceIDs[1] = MeshDescription.CreateVertexInstance(VertexIDs[P1]);
+		VertexInstanceIDs[2] = MeshDescription.CreateVertexInstance(VertexIDs[P2]);
+		VertexInstanceIDs[3] = MeshDescription.CreateVertexInstance(VertexIDs[P3]);
+
+		TArray<FEdgeID> EdgeIDs;
+		EdgeIDs.Reserve(4);
+
+		FPolygonID PolygonID = MeshDescription.CreatePolygon(PolygonGroup, VertexInstanceIDs, &EdgeIDs);
+
+		for (FEdgeID EdgeID : EdgeIDs)
+		{
+			GetEdgeHardnesses()[EdgeID] = true;
+		}
+
+		return PolygonID;
+	};
+
+	PolygonID_PlusX = MakePolygon(0, 1, 7, 6);
+	PolygonID_MinusX = MakePolygon(2, 3, 5, 4);
+	PolygonID_PlusY = MakePolygon(1, 2, 4, 7);
+	PolygonID_MinusY = MakePolygon(3, 0, 6, 5);
+	PolygonID_PlusZ = MakePolygon(1, 0, 3, 2);
+	PolygonID_MinusZ = MakePolygon(5, 4, 7, 6);
+
+	MeshDescription.PolygonAttributes().RegisterAttribute<FVector>(MeshAttribute::Polygon::Normal, 1, FVector::ZeroVector, EMeshAttributeFlags::Transient);
+	MeshDescription.PolygonAttributes().RegisterAttribute<FVector>(MeshAttribute::Polygon::Tangent, 1, FVector::ZeroVector, EMeshAttributeFlags::Transient);
+	MeshDescription.PolygonAttributes().RegisterAttribute<FVector>(MeshAttribute::Polygon::Binormal, 1, FVector::ZeroVector, EMeshAttributeFlags::Transient);
+	MeshDescription.PolygonAttributes().RegisterAttribute<FVector>(MeshAttribute::Polygon::Center, 1, FVector::ZeroVector, EMeshAttributeFlags::Transient);
+
+	FStaticMeshOperations::ComputePolygonTangentsAndNormals(MeshDescription);
+	FStaticMeshOperations::ComputeTangentsAndNormals(MeshDescription, EComputeNTBsFlags::Normals | EComputeNTBsFlags::Tangents);
+}
+
+
diff --git a/Engine/Source/Runtime/StaticMeshDescription/Private/StaticMeshOperations.cpp b/Engine/Source/Runtime/StaticMeshDescription/Private/StaticMeshOperations.cpp
new file mode 100644
--- /dev/null
+++ b/Engine/Source/Runtime/StaticMeshDescription/Private/StaticMeshOperations.cpp
@@ -0,0 +1,778 @@
+// Copyright 1998-2019 Epic Games, Inc. All Rights Reserved.
+
+#include "StaticMeshOperations.h"
+#include "StaticMeshAttributes.h"
+
+
+static bool GetPolygonTangentsAndNormals(FMeshDescription& MeshDescription,
+										 FPolygonID PolygonID,
+										 float ComparisonThreshold, 
+										 TVertexAttributesRef<const FVector> VertexPositions,
+										 TVertexInstanceAttributesRef<const FVector2D> VertexUVs,
+										 TPolygonAttributesRef<FVector> PolygonNormals,
+										 TPolygonAttributesRef<FVector> PolygonTangents,
+										 TPolygonAttributesRef<FVector> PolygonBinormals,
+										 TPolygonAttributesRef<FVector> PolygonCenters)
+{
+	bool bValidNTBs = true;
+
+	// Calculate the center of this polygon
+	FVector Center = FVector::ZeroVector;
+	const TArray<FVertexInstanceID>& VertexInstanceIDs = MeshDescription.GetPolygonVertexInstances(PolygonID);
+	for (const FVertexInstanceID VertexInstanceID : VertexInstanceIDs)
+	{
+		Center += VertexPositions[MeshDescription.GetVertexInstanceVertex(VertexInstanceID)];
+	}
+	Center /= float(VertexInstanceIDs.Num());
+
+	// Calculate the tangent basis for the polygon, based on the average of all constituent triangles
+	FVector Normal = FVector::ZeroVector;
+	FVector Tangent = FVector::ZeroVector;
+	FVector Binormal = FVector::ZeroVector;
+
+	for (const FTriangleID TriangleID : MeshDescription.GetPolygonTriangleIDs(PolygonID))
+	{
+		TArrayView<const FVertexInstanceID> TriangleVertexInstances = MeshDescription.GetTriangleVertexInstances(TriangleID);
+		const FVertexID VertexID0 = MeshDescription.GetVertexInstanceVertex(TriangleVertexInstances[0]);
+		const FVertexID VertexID1 = MeshDescription.GetVertexInstanceVertex(TriangleVertexInstances[1]);
+		const FVertexID VertexID2 = MeshDescription.GetVertexInstanceVertex(TriangleVertexInstances[2]);
+
+		const FVector DPosition1 = VertexPositions[VertexID1] - VertexPositions[VertexID0];
+		const FVector DPosition2 = VertexPositions[VertexID2] - VertexPositions[VertexID0];
+
+		const FVector2D DUV1 = VertexUVs[TriangleVertexInstances[1]] - VertexUVs[TriangleVertexInstances[0]];
+		const FVector2D DUV2 = VertexUVs[TriangleVertexInstances[2]] - VertexUVs[TriangleVertexInstances[0]];
+
+		// We have a left-handed coordinate system, but a counter-clockwise winding order
+		// Hence normal calculation has to take the triangle vectors cross product in reverse.
+		FVector TmpNormal = FVector::CrossProduct(DPosition2, DPosition1);
+		if (!TmpNormal.IsNearlyZero(ComparisonThreshold) && !TmpNormal.ContainsNaN())
+		{
+			Normal += TmpNormal;
+			// ...and tangent space seems to be right-handed.
+			const float DetUV = FVector2D::CrossProduct(DUV1, DUV2);
+			const float InvDetUV = (DetUV == 0.0f) ? 0.0f : 1.0f / DetUV;
+
+			Tangent += (DPosition1 * DUV2.Y - DPosition2 * DUV1.Y) * InvDetUV;
+			Binormal += (DPosition2 * DUV1.X - DPosition1 * DUV2.X) * InvDetUV;
+		}
+		else
+		{
+			// The polygon is degenerated
+			bValidNTBs = false;
+		}
+	}
+
+	PolygonNormals[PolygonID] = Normal.GetSafeNormal();
+	PolygonTangents[PolygonID] = Tangent.GetSafeNormal();
+	PolygonBinormals[PolygonID] = Binormal.GetSafeNormal();
+	PolygonCenters[PolygonID] = Center;
+
+	return bValidNTBs;
+}
+
+
+static TArray<FPolygonID> GetPolygonsInSameSoftEdgedGroupAsPolygon(const FMeshDescription& MeshDescription, FPolygonID PolygonID, const TArray<FPolygonID>& CandidatePolygonIDs, const TArray<FEdgeID>& SoftEdgeIDs)
+{
+	// The aim of this method is:
+	// - given a polygon ID,
+	// - given a set of candidate polygons connected to the same vertex (which should include the polygon ID),
+	// - given a set of soft edges connected to the same vertex,
+	// return the polygon IDs which form an adjacent run without crossing a hard edge.
+
+	TArray<FPolygonID> OutPolygonIDs;
+
+	// Maintain a list of polygon IDs to be examined. Adjacents are added to the list if suitable.
+	// Add the start poly here.
+	TArray<FPolygonID> PolygonsToCheck;
+	PolygonsToCheck.Reset(CandidatePolygonIDs.Num());
+	PolygonsToCheck.Add(PolygonID);
+
+	int32 Index = 0;
+	while (Index < PolygonsToCheck.Num())
+	{
+		const FPolygonID PolygonToCheck = PolygonsToCheck[Index];
+		Index++;
+
+		if (CandidatePolygonIDs.Contains(PolygonToCheck))
+		{
+			OutPolygonIDs.Add(PolygonToCheck);
+
+			// Now look at its adjacent polygons. If they are joined by a soft edge which includes the vertex we're interested in, we want to consider them.
+			// We take a shortcut by doing this process in reverse: we already know all the soft edges we are interested in, so check if any of them
+			// have the current polygon as an adjacent.
+			for (const FEdgeID SoftEdgeID : SoftEdgeIDs)
+			{
+				const TArray<FPolygonID>& EdgeConnectedPolygons = MeshDescription.GetEdgeConnectedPolygons(SoftEdgeID);
+				if (EdgeConnectedPolygons.Contains(PolygonToCheck))
+				{
+					for (const FPolygonID AdjacentPolygon : EdgeConnectedPolygons)
+					{
+						// Only add new polygons which haven't yet been added to the list. This prevents circular runs of polygons triggering infinite loops.
+						PolygonsToCheck.AddUnique(AdjacentPolygon);
+					}
+				}
+			}
+		}
+	}
+
+	return OutPolygonIDs;
+}
+
+
+static TArray<FEdgeID> GetConnectedSoftEdges(const FMeshDescription& MeshDescription, FVertexID VertexID)
+{
+	TArray<FEdgeID> OutConnectedSoftEdges;
+	OutConnectedSoftEdges.Reserve(MeshDescription.GetNumVertexConnectedEdges(VertexID));
+
+	FStaticMeshConstAttributes Attributes(MeshDescription);
+	TEdgeAttributesConstRef<bool> EdgeHardnesses = Attributes.GetEdgeHardnesses();
+	for (const FEdgeID ConnectedEdgeID : MeshDescription.GetVertexConnectedEdges(VertexID))
+	{
+		if (!EdgeHardnesses[ConnectedEdgeID])
+		{
+			OutConnectedSoftEdges.Add(ConnectedEdgeID);
+		}
+	}
+
+	return OutConnectedSoftEdges;
+}
+
+
+static TArray<FPolygonID> GetVertexConnectedPolygonsInSameSoftEdgedGroup(const FMeshDescription& MeshDescription, FVertexID VertexID, FPolygonID PolygonID)
+{
+	// The aim here is to determine which polygons form part of the same soft edged group as the polygons attached to this vertex.
+	// They should all contribute to the final vertex instance normal.
+
+	// Get all polygons connected to this vertex.
+	TArray<FPolygonID> ConnectedPolygons = MeshDescription.GetVertexConnectedPolygons(VertexID);
+
+	// Cache a list of all soft edges which share this vertex.
+	// We're only interested in finding adjacent polygons which are not the other side of a hard edge.
+	TArray<FEdgeID> ConnectedSoftEdges = GetConnectedSoftEdges(MeshDescription, VertexID);
+
+	return GetPolygonsInSameSoftEdgedGroupAsPolygon(MeshDescription, PolygonID, ConnectedPolygons, ConnectedSoftEdges);
+}
+
+
+static void GetTangentsAndNormals(FMeshDescription& MeshDescription,
+								  FVertexInstanceID VertexInstanceID,
+								  EComputeNTBsFlags ComputeNTBsOptions,
+								  TPolygonAttributesRef<const FVector> PolygonNormals,
+								  TPolygonAttributesRef<const FVector> PolygonTangents,
+								  TPolygonAttributesRef<const FVector> PolygonBinormals,
+								  TVertexInstanceAttributesRef<FVector> VertexNormals,
+								  TVertexInstanceAttributesRef<FVector> VertexTangents,
+								  TVertexInstanceAttributesRef<float> VertexBinormalSigns)
+{
+	bool bComputeNormals = !!(ComputeNTBsOptions & EComputeNTBsFlags::Normals);
+	bool bComputeTangents = !!(ComputeNTBsOptions & EComputeNTBsFlags::Tangents);
+	bool bUseWeightedNormals = !!(ComputeNTBsOptions & EComputeNTBsFlags::WeightedNTBs);
+
+	FVector Normal = FVector::ZeroVector;
+	FVector Tangent = FVector::ZeroVector;
+	FVector Binormal = FVector::ZeroVector;
+
+	FVector& NormalRef = VertexNormals[VertexInstanceID];
+	FVector& TangentRef = VertexTangents[VertexInstanceID];
+	float& BinormalRef = VertexBinormalSigns[VertexInstanceID];
+
+	if (!bComputeNormals && !bComputeTangents)
+	{
+		// Nothing to compute
+		return;
+	}
+
+	const FVertexID VertexID = MeshDescription.GetVertexInstanceVertex(VertexInstanceID);
+
+	if (bComputeNormals || NormalRef.IsNearlyZero())
+	{
+		// Get all polygons connected to this vertex instance
+		const TArray<FPolygonID>& VertexInstanceConnectedPolygons = MeshDescription.GetVertexInstanceConnectedPolygons(VertexInstanceID);
+		check(VertexInstanceConnectedPolygons.Num() > 0);
+
+		// Add also any in the same smoothing group connected to a different vertex instance
+		// (as they still have influence over the normal).
+		TArray<FPolygonID> AllConnectedPolygons = GetVertexConnectedPolygonsInSameSoftEdgedGroup(MeshDescription, VertexID, VertexInstanceConnectedPolygons[0]);
+
+		// The vertex instance normal is computed as a sum of all connected polygons' normals, weighted by the angle they make with the vertex
+		for (const FPolygonID ConnectedPolygonID : AllConnectedPolygons)
+		{
+			const float Angle = bUseWeightedNormals ? MeshDescription.GetPolygonCornerAngleForVertex(ConnectedPolygonID, VertexID) : 1.0f;
+
+			Normal += PolygonNormals[ConnectedPolygonID] * Angle;
+
+			// If this polygon is actually connected to the vertex instance we're processing, also include its contributions towards the tangent
+			if (VertexInstanceConnectedPolygons.Contains(ConnectedPolygonID))
+			{
+				Tangent += PolygonTangents[ConnectedPolygonID] * Angle;
+				Binormal += PolygonBinormals[ConnectedPolygonID] * Angle;
+			}
+		}
+
+		// Normalize Normal
+		Normal = Normal.GetSafeNormal();
+	}
+	else
+	{
+		// We use existing normals so just use all polygons having a vertex instance at the same location sharing the same normals
+		Normal = NormalRef;
+		TArray<FVertexInstanceID> VertexInstanceIDs = MeshDescription.GetVertexVertexInstances(VertexID);
+		for (const FVertexInstanceID& ConnectedVertexInstanceID : VertexInstanceIDs)
+		{
+			if (ConnectedVertexInstanceID != VertexInstanceID && !VertexNormals[ConnectedVertexInstanceID].Equals(Normal))
+			{
+				continue;
+			}
+
+			const TArray<FPolygonID>& ConnectedPolygons = MeshDescription.GetVertexInstanceConnectedPolygons(ConnectedVertexInstanceID);
+			for (const FPolygonID ConnectedPolygonID : ConnectedPolygons)
+			{
+				const float Angle = bUseWeightedNormals ? MeshDescription.GetPolygonCornerAngleForVertex(ConnectedPolygonID, VertexID) : 1.0f;
+
+				// If this polygon is actually connected to the vertex instance we're processing, also include its contributions towards the tangent
+				Tangent += PolygonTangents[ConnectedPolygonID] * Angle;
+				Binormal += PolygonBinormals[ConnectedPolygonID] * Angle;
+			}
+		}
+	}
+
+
+	float BinormalSign = 1.0f;
+	if (bComputeTangents)
+	{
+		// Make Tangent orthonormal to Normal.
+		// This is a quicker method than normalizing Tangent, taking the cross product Normal X Tangent, and then a further cross product with that result
+		Tangent = (Tangent - Normal * FVector::DotProduct(Normal, Tangent)).GetSafeNormal();
+
+		// Calculate binormal sign
+		BinormalSign = (FVector::DotProduct(FVector::CrossProduct(Normal, Tangent), Binormal) < 0.0f) ? -1.0f : 1.0f;
+	}
+
+	// Set the value that need to be set
+	if (NormalRef.IsNearlyZero())
+	{
+		NormalRef = Normal;
+	}
+
+	if (bComputeTangents)
+	{
+		if (TangentRef.IsNearlyZero())
+		{
+			TangentRef = Tangent;
+		}
+
+		if (FMath::IsNearlyZero(BinormalRef))
+		{
+			BinormalRef = BinormalSign;
+		}
+	}
+}
+
+
+
+void FStaticMeshOperations::ComputePolygonTangentsAndNormals(FMeshDescription& MeshDescription, TArrayView<const FPolygonID> PolygonIDs, float ComparisonThreshold)
+{
+	FStaticMeshAttributes Attributes(MeshDescription);
+	TVertexAttributesRef<const FVector> VertexPositions = Attributes.GetVertexPositions();
+	TVertexInstanceAttributesRef<const FVector2D> VertexUVs = Attributes.GetVertexInstanceUVs();
+	TPolygonAttributesRef<FVector> PolygonNormals = Attributes.GetPolygonNormals();
+	TPolygonAttributesRef<FVector> PolygonTangents = Attributes.GetPolygonTangents();
+	TPolygonAttributesRef<FVector> PolygonBinormals = Attributes.GetPolygonBinormals();
+	TPolygonAttributesRef<FVector> PolygonCenters = Attributes.GetPolygonCenters();
+
+	TArray<FPolygonID> DegeneratePolygonIDs;
+	for (const FPolygonID PolygonID : PolygonIDs)
+	{
+		if (!GetPolygonTangentsAndNormals(MeshDescription, PolygonID, ComparisonThreshold, VertexPositions, VertexUVs, PolygonNormals, PolygonTangents, PolygonBinormals, PolygonCenters))
+		{
+			DegeneratePolygonIDs.Add(PolygonID);
+		}
+	}
+
+	// Remove degenerated polygons
+	// Delete the degenerated polygons. The array is filled only if the remove degenerated option is turned on.
+	if (DegeneratePolygonIDs.Num() > 0)
+	{
+		TArray<FEdgeID> OrphanedEdges;
+		TArray<FVertexInstanceID> OrphanedVertexInstances;
+		TArray<FPolygonGroupID> OrphanedPolygonGroups;
+		TArray<FVertexID> OrphanedVertices;
+
+		for (FPolygonID PolygonID : DegeneratePolygonIDs)
+		{
+			MeshDescription.DeletePolygon(PolygonID, &OrphanedEdges, &OrphanedVertexInstances, &OrphanedPolygonGroups);
+		}
+		for (FPolygonGroupID PolygonGroupID : OrphanedPolygonGroups)
+		{
+			MeshDescription.DeletePolygonGroup(PolygonGroupID);
+		}
+		for (FVertexInstanceID VertexInstanceID : OrphanedVertexInstances)
+		{
+			MeshDescription.DeleteVertexInstance(VertexInstanceID, &OrphanedVertices);
+		}
+		for (FEdgeID EdgeID : OrphanedEdges)
+		{
+			MeshDescription.DeleteEdge(EdgeID, &OrphanedVertices);
+		}
+		for (FVertexID VertexID : OrphanedVertices)
+		{
+			MeshDescription.DeleteVertex(VertexID);
+		}
+		// Compact and Remap IDs so we have clean ID from 0 to n since we just erase some polygons
+		FElementIDRemappings RemappingInfos;
+		MeshDescription.Compact(RemappingInfos);
+	}
+}
+
+void FStaticMeshOperations::ComputePolygonTangentsAndNormals(FMeshDescription& MeshDescription, float ComparisonThreshold)
+{
+	TArray<FPolygonID> PolygonsToComputeNTBs;
+	PolygonsToComputeNTBs.Reserve(MeshDescription.Polygons().Num());
+
+	for (const FPolygonID PolygonID : MeshDescription.Polygons().GetElementIDs())
+	{
+		PolygonsToComputeNTBs.Add(PolygonID);
+	}
+
+	ComputePolygonTangentsAndNormals(MeshDescription, PolygonsToComputeNTBs, ComparisonThreshold);
+}
+
+void FStaticMeshOperations::ComputeTangentsAndNormals(FMeshDescription& MeshDescription, TArrayView<const FVertexInstanceID> VertexInstanceIDs, EComputeNTBsFlags ComputeNTBsOptions)
+{
+	FStaticMeshAttributes Attributes(MeshDescription);
+	TPolygonAttributesRef<const FVector> PolygonNormals = Attributes.GetPolygonNormals();
+	TPolygonAttributesRef<const FVector> PolygonTangents = Attributes.GetPolygonTangents();
+	TPolygonAttributesRef<const FVector> PolygonBinormals = Attributes.GetPolygonBinormals();
+	TVertexInstanceAttributesRef<FVector> VertexNormals = Attributes.GetVertexInstanceNormals();
+	TVertexInstanceAttributesRef<FVector> VertexTangents = Attributes.GetVertexInstanceTangents();
+	TVertexInstanceAttributesRef<float> VertexBinormalSigns = Attributes.GetVertexInstanceBinormalSigns();
+
+	for (const FVertexInstanceID VertexInstanceID : VertexInstanceIDs)
+	{
+		GetTangentsAndNormals(MeshDescription, VertexInstanceID, ComputeNTBsOptions, PolygonNormals, PolygonTangents, PolygonBinormals, VertexNormals, VertexTangents, VertexBinormalSigns);
+	}
+}
+
+void FStaticMeshOperations::ComputeTangentsAndNormals(FMeshDescription& MeshDescription, EComputeNTBsFlags ComputeNTBsOptions)
+{
+	FStaticMeshAttributes Attributes(MeshDescription);
+	TPolygonAttributesRef<const FVector> PolygonNormals = Attributes.GetPolygonNormals();
+	TPolygonAttributesRef<const FVector> PolygonTangents = Attributes.GetPolygonTangents();
+	TPolygonAttributesRef<const FVector> PolygonBinormals = Attributes.GetPolygonBinormals();
+	TVertexInstanceAttributesRef<FVector> VertexNormals = Attributes.GetVertexInstanceNormals();
+	TVertexInstanceAttributesRef<FVector> VertexTangents = Attributes.GetVertexInstanceTangents();
+	TVertexInstanceAttributesRef<float> VertexBinormalSigns = Attributes.GetVertexInstanceBinormalSigns();
+
+	for (const FVertexInstanceID VertexInstanceID : MeshDescription.VertexInstances().GetElementIDs())
+	{
+		GetTangentsAndNormals(MeshDescription, VertexInstanceID, ComputeNTBsOptions, PolygonNormals, PolygonTangents, PolygonBinormals, VertexNormals, VertexTangents, VertexBinormalSigns);
+	}
+}
+
+void FStaticMeshOperations::DetermineEdgeHardnessesFromVertexInstanceNormals(FMeshDescription& MeshDescription, float Tolerance)
+{
+	FStaticMeshAttributes Attributes(MeshDescription);
+
+	TVertexInstanceAttributesRef<const FVector> VertexNormals = Attributes.GetVertexInstanceNormals();
+	TEdgeAttributesRef<bool> EdgeHardnesses = Attributes.GetEdgeHardnesses();
+
+	// Holds unique vertex instance IDs for a given edge vertex
+	TArray<FVertexInstanceID> UniqueVertexInstanceIDs;
+
+	for (const FEdgeID EdgeID : MeshDescription.Edges().GetElementIDs())
+	{
+		// Get list of polygons connected to this edge
+		TArray<FPolygonID, TInlineAllocator<2>> ConnectedPolygonIDs = MeshDescription.GetEdgeConnectedPolygons<TInlineAllocator<2>>(EdgeID);
+		if (ConnectedPolygonIDs.Num() == 0)
+		{
+			// What does it mean if an edge has no connected polygons? For now we just skip it
+			continue;
+		}
+
+		// Assume by default that the edge is soft - but as soon as any vertex instance belonging to a connected polygon
+		// has a distinct normal from the others (within the given tolerance), we mark it as hard.
+		// The exception is if an edge has exactly one connected polygon: in this case we automatically deem it a hard edge.
+		bool bEdgeIsHard = (ConnectedPolygonIDs.Num() == 1);
+
+		// Examine vertices on each end of the edge, if we haven't yet identified it as 'hard'
+		for (int32 VertexIndex = 0; !bEdgeIsHard && VertexIndex < 2; ++VertexIndex)
+		{
+			const FVertexID VertexID = MeshDescription.GetEdgeVertex(EdgeID, VertexIndex);
+
+			const int32 ReservedElements = 4;
+			UniqueVertexInstanceIDs.Reset(ReservedElements);
+
+			// Get a list of all vertex instances for this vertex which form part of any polygon connected to the edge
+			for (const FVertexInstanceID VertexInstanceID : MeshDescription.GetVertexVertexInstances(VertexID))
+			{
+				for (const FPolygonID PolygonID : MeshDescription.GetVertexInstanceConnectedPolygons<TInlineAllocator<8>>(VertexInstanceID))
+				{
+					if (ConnectedPolygonIDs.Contains(PolygonID))
+					{
+						UniqueVertexInstanceIDs.AddUnique(VertexInstanceID);
+						break;
+					}
+				}
+			}
+			check(UniqueVertexInstanceIDs.Num() > 0);
+
+			// First unique vertex instance is used as a reference against which the others are compared.
+			// (not a perfect approach: really the 'median' should be used as a reference)
+			const FVector ReferenceNormal = VertexNormals[UniqueVertexInstanceIDs[0]];
+			for (int32 Index = 1; Index < UniqueVertexInstanceIDs.Num(); ++Index)
+			{
+				if (!VertexNormals[UniqueVertexInstanceIDs[Index]].Equals(ReferenceNormal, Tolerance))
+				{
+					bEdgeIsHard = true;
+					break;
+				}
+			}
+		}
+
+		EdgeHardnesses[EdgeID] = bEdgeIsHard;
+	}
+}
+
+
+void FStaticMeshOperations::ExtrudePolygons(FMeshDescription& MeshDescription, const TArray<FPolygonID>& PolygonIDs, float ExtrudeDistance, bool bKeepNeighborsTogether, TArray<FPolygonID>& OutNewExtrudedFrontPolygons)
+{
+	FStaticMeshAttributes Attributes(MeshDescription);
+	TVertexAttributesRef<FVector> VertexPositions = Attributes.GetVertexPositions();
+
+	OutNewExtrudedFrontPolygons.Reset();
+
+	// Convert our incoming polygon array to a TSet so we can lookup quickly to see which polygons in the mesh sare members of the set
+	TSet<FPolygonID> PolygonsSet;
+	PolygonsSet.Append(PolygonIDs);
+
+	TArray<FPolygonID> AllNewPolygons;
+
+//	static TArray<FAttributesForEdge> AttributesForEdges;
+//	AttributesForEdges.Reset();
+
+//	static TArray<FAttributesForVertex> AttributesForVertices;
+//	AttributesForVertices.Reset();
+
+	// First, let's figure out which of the polygons we were asked to extrude share edges or vertices.  We'll keep those edges intact!
+	TMap<FEdgeID, uint32> EdgeUsageCounts;	// Maps an edge ID to the number of times it is referenced by the incoming polygons
+	TSet<FVertexID> UniqueVertexIDs;
+
+	for (const FPolygonID PolygonID : PolygonIDs)
+	{
+		for (const FEdgeID EdgeID : MeshDescription.GetPolygonPerimeterEdges(PolygonID))
+		{
+			if (uint32* EdgeUsageCountPtr = EdgeUsageCounts.Find(EdgeID))
+			{
+				++(*EdgeUsageCountPtr);
+			}
+			else
+			{
+				EdgeUsageCounts.Add(EdgeID, 1);
+			}
+		}
+
+		UniqueVertexIDs.Append(MeshDescription.GetPolygonVertices(PolygonID));
+	}
+
+	int32 NumVerticesToCreate = UniqueVertexIDs.Num();
+
+	// Create new vertices for all of the extruded polygons
+	TArray<FVertexID> ExtrudedVertexIDs;
+	ExtrudedVertexIDs.Reserve(NumVerticesToCreate);
+	ExtrudedVertexIDs.Add(MeshDescription.CreateVertex());
+
+	int32 NextAvailableExtrudedVertexIDNumber = 0;
+
+	TMap<FVertexID, FVertexID> VertexIDToExtrudedCopy;
+	#if 0
+	for (int32 PassIndex = 0; PassIndex < 2; ++PassIndex)
+	{
+		// Extrude all of the shared edges first, then do the non-shared edges.  This is to make sure that a vertex doesn't get offset
+		// without taking into account all of the connected polygons in our set
+		const bool bIsExtrudingSharedEdges = (PassIndex == 0);
+
+		for (const FPolygonID PolygonID : PolygonIDs)
+		{
+			const FPolygonGroupID PolygonGroupID = MeshDescription.GetPolygonPolygonGroup(PolygonID);
+
+			if (!bKeepNeighborsTogether)
+			{
+				VertexIDToExtrudedCopy.Reset();
+			}
+
+			// Map all of the edge vertices to their new extruded counterpart
+			TArray<FVertexID> PolygonVertices = MeshDescription.GetPolygonVertices(PolygonID);
+			TArray<FEdgeID> PolygonPerimeterEdges = MeshDescription.GetPolygonPerimeterEdges(PolygonID);
+			for (const FEdgeID EdgeID : PolygonPerimeterEdges)
+			{
+				// @todo mesheditor perf: We can change GetPolygonPerimeterEdges() to have a version that returns whether winding is reversed or not, and avoid this call entirely.
+				// @todo mesheditor perf: O(log N^2) iteration here. For every edge, for every edge up to this index.  Need to clean this up. 
+				//		--> Also, there are quite a few places where we are stepping through edges in perimeter-order.  We need to have a nice way to walk that.
+				bool bEdgeWindingIsReversedForPolygon;
+				const FEdgeID EdgeID = this->GetPolygonPerimeterEdge( PolygonID, PerimeterEdgeNumber, /* Out */ bEdgeWindingIsReversedForPolygon );
+
+				const bool bIsSharedEdge = bKeepNeighborsTogether && EdgeUsageCounts[ EdgeID ] > 1;
+				if( bIsSharedEdge == bIsExtrudingSharedEdges )
+				{
+					FVertexID EdgeVertexIDs[ 2 ];
+					GetEdgeVertices( EdgeID, /* Out */ EdgeVertexIDs[ 0 ], /* Out */ EdgeVertexIDs[ 1 ] );
+					if( bEdgeWindingIsReversedForPolygon )
+					{
+						::Swap( EdgeVertexIDs[ 0 ], EdgeVertexIDs[ 1 ] );
+					}
+
+					if( !bIsSharedEdge )
+					{
+						// After extruding, all of the edges of the original polygon become hard edges
+						AttributesForEdges.Emplace();
+						FAttributesForEdge& AttributesForEdge = AttributesForEdges.Last();
+						AttributesForEdge.EdgeID = EdgeID;
+						AttributesForEdge.EdgeAttributes.Attributes.Emplace( MeshAttribute::Edge::IsHard, 0, FMeshElementAttributeValue( true ) );
+					}
+
+					FVertexID ExtrudedEdgeVertexIDs[ 2 ];
+					for( int32 EdgeVertexNumber = 0; EdgeVertexNumber < 2; ++EdgeVertexNumber )
+					{
+						const FVertexID EdgeVertexID = EdgeVertexIDs[ EdgeVertexNumber ];
+						FVertexID* ExtrudedEdgeVertexIDPtr = VertexIDToExtrudedCopy.Find( EdgeVertexID );
+
+						// @todo mesheditor extrude: Ideally we would detect whether the vertex that was already extruded came from a edge
+						// from a polygon that does not actually share an edge with any polygons this polygon shares an edge with.  This
+						// would avoid the problem where extruding two polygons that are connected only by a vertex are not extruded
+						// separately.
+						const bool bVertexIsSharedByAnEdgeOfAnotherSelectedPolygon = false;
+						if( ExtrudedEdgeVertexIDPtr != nullptr)// && !bVertexIsSharedByAnEdgeOfAnotherSelectedPolygon )
+						{
+							ExtrudedEdgeVertexIDs[ EdgeVertexNumber ] = *ExtrudedEdgeVertexIDPtr;
+						}
+						else
+						{
+							// Create a copy of this vertex for the extruded face
+							const FVertexID ExtrudedVertexID = ExtrudedVertexIDs[ NextAvailableExtrudedVertexIDNumber++ ];
+
+							ExtrudedEdgeVertexIDPtr = &VertexIDToExtrudedCopy.Add( EdgeVertexID, ExtrudedVertexID );
+							// Push the vertex out along the polygon's normal
+							const FVector OriginalVertexPosition = VertexPositions[ EdgeVertexID ];
+
+							FVector ExtrudedVertexPosition;
+							if( bIsSharedEdge )
+							{
+								// Get all of the polygons that share this edge that were part of the set of polygons passed in.  We'll
+								// generate an extrude direction that's the average of those polygon normals.
+								FVector ExtrudeDirection = FVector::ZeroVector;
+
+								static TArray<FPolygonID> ConnectedPolygonIDs;
+								GetVertexConnectedPolygons( EdgeVertexID, /* Out */ ConnectedPolygonIDs );
+
+								static TArray<FPolygonID> NeighborPolygonIDs;
+								NeighborPolygonIDs.Reset();
+								for( const FPolygonID ConnectedPolygonID : ConnectedPolygonIDs )
+								{
+									// We only care about polygons that are members of the set of polygons we were asked to extrude
+									if( PolygonsSet.Contains( ConnectedPolygonID ) )
+									{
+										NeighborPolygonIDs.Add( ConnectedPolygonID );
+
+										// We'll need this polygon's normal to figure out where to put the extruded copy of the polygon
+										const FVector NeighborPolygonNormal = ComputePolygonNormal( ConnectedPolygonID );
+										ExtrudeDirection += NeighborPolygonNormal;
+									}
+								}
+								ExtrudeDirection.Normalize();
+
+
+								// OK, we have the direction to extrude for this vertex.  Now we need to know how far to extrude.  We'll
+								// loop over all of the neighbor polygons to this vertex, and choose the closest intersection point with our
+								// vertex's extrude direction and the neighbor polygon's extruded plane
+								FVector ClosestIntersectionPointWithExtrudedPlanes;
+								float ClosestIntersectionDistanceSquared = TNumericLimits<float>::Max();
+
+								for( const FPolygonID NeighborPolygonID : NeighborPolygonIDs )
+								{
+									const FPlane NeighborPolygonPlane = ComputePolygonPlane( NeighborPolygonID );
+
+									// Push the plane out
+									const FPlane ExtrudedPlane = [NeighborPolygonPlane, ExtrudeDistance]
+									{ 
+										FPlane NewPlane = NeighborPolygonPlane;
+										NewPlane.W += ExtrudeDistance; 
+										return NewPlane; 
+									}();
+
+									// Is this the closest intersection point so far?
+									const FVector IntersectionPointWithExtrudedPlane = FMath::RayPlaneIntersection( OriginalVertexPosition, ExtrudeDirection, ExtrudedPlane );
+									const float IntersectionDistanceSquared = FVector::DistSquared( OriginalVertexPosition, IntersectionPointWithExtrudedPlane );
+									if( IntersectionDistanceSquared < ClosestIntersectionDistanceSquared )
+									{
+										ClosestIntersectionPointWithExtrudedPlanes = IntersectionPointWithExtrudedPlane;
+										ClosestIntersectionDistanceSquared = IntersectionDistanceSquared;
+									}
+								}
+
+								ExtrudedVertexPosition = ClosestIntersectionPointWithExtrudedPlanes;
+							}
+							else
+							{
+								// We'll need this polygon's normal to figure out where to put the extruded copy of the polygon
+								const FVector PolygonNormal = ComputePolygonNormal( PolygonID );
+								ExtrudedVertexPosition = OriginalVertexPosition + ExtrudeDistance * PolygonNormal;
+							}
+
+							// Fill in the vertex
+							AttributesForVertices.Emplace();
+							FAttributesForVertex& AttributesForVertex = AttributesForVertices.Last();
+							AttributesForVertex.VertexID = ExtrudedVertexID;
+							AttributesForVertex.VertexAttributes.Attributes.Emplace( MeshAttribute::Vertex::Position, 0, FMeshElementAttributeValue( ExtrudedVertexPosition ) );
+						}
+						ExtrudedEdgeVertexIDs[ EdgeVertexNumber ] = *ExtrudedEdgeVertexIDPtr;
+					}
+
+					if( !bIsSharedEdge )
+					{
+						static TArray<FVertexAndAttributes> NewSidePolygonVertices;
+						NewSidePolygonVertices.Reset( 4 );
+						NewSidePolygonVertices.SetNum( 4, false );	// Always four edges in an extruded face
+
+						NewSidePolygonVertices[ 0 ].VertexID = EdgeVertexIDs[ 1 ];
+						NewSidePolygonVertices[ 1 ].VertexID = EdgeVertexIDs[ 0 ];
+						NewSidePolygonVertices[ 2 ].VertexID = ExtrudedEdgeVertexIDs[ 0 ];
+						NewSidePolygonVertices[ 3 ].VertexID = ExtrudedEdgeVertexIDs[ 1 ];
+
+						// Get vertex instance IDs on this polygon corresponding to the edge start/end vertices
+						const FVertexInstanceID EdgeVertexInstanceID0 = GetMeshDescription()->GetVertexInstanceForPolygonVertex( PolygonID, EdgeVertexIDs[ 0 ] );
+						const FVertexInstanceID EdgeVertexInstanceID1 = GetMeshDescription()->GetVertexInstanceForPolygonVertex( PolygonID, EdgeVertexIDs[ 1 ] );
+
+						BackupAllAttributes( NewSidePolygonVertices[ 0 ].PolygonVertexAttributes, GetMeshDescription()->VertexInstanceAttributes(), EdgeVertexInstanceID1 );
+						BackupAllAttributes( NewSidePolygonVertices[ 1 ].PolygonVertexAttributes, GetMeshDescription()->VertexInstanceAttributes(), EdgeVertexInstanceID0 );
+						BackupAllAttributes( NewSidePolygonVertices[ 2 ].PolygonVertexAttributes, GetMeshDescription()->VertexInstanceAttributes(), EdgeVertexInstanceID0 );
+						BackupAllAttributes( NewSidePolygonVertices[ 3 ].PolygonVertexAttributes, GetMeshDescription()->VertexInstanceAttributes(), EdgeVertexInstanceID1 );
+
+						FPolygonID NewSidePolygonID;	// Filled in below
+						{
+							static TArray<FPolygonToCreate> PolygonsToCreate;
+							PolygonsToCreate.Reset();
+
+							// Create the polygon
+							// @todo mesheditor perf: Ideally we support creating multiple polygons at once and batching up the work
+							FPolygonToCreate& PolygonToCreate = *new( PolygonsToCreate ) FPolygonToCreate();
+							PolygonToCreate.PolygonGroupID = PolygonGroupID;
+							PolygonToCreate.PerimeterVertices = NewSidePolygonVertices;	// @todo mesheditor perf: Copying static array here, ideally allocations could be avoided
+							PolygonToCreate.PolygonEdgeHardness = EPolygonEdgeHardness::AllEdgesHard;
+
+							static TArray<FPolygonID> NewPolygonIDs;
+							NewPolygonIDs.Reset();
+							static TArray<FEdgeID> NewEdgeIDs;
+							NewEdgeIDs.Reset();
+							CreatePolygons( PolygonsToCreate, /* Out */ NewPolygonIDs, /* Out */ NewEdgeIDs );	// @todo mesheditor perf: Extra allocatons/copies: Ideally MoveTemp() here but we can't move a STATIC local!
+
+							NewSidePolygonID = NewPolygonIDs[ 0 ];
+						}
+						AllNewPolygons.Add( NewSidePolygonID );
+					}
+				}
+			}
+		}
+	}
+
+	for( int32 PolygonIter = 0; PolygonIter < PolygonIDs.Num(); ++PolygonIter )
+	{
+		const FPolygonID PolygonID = PolygonIDs[ PolygonIter ];
+		const FPolygonGroupID PolygonGroupID = GetGroupForPolygon( PolygonID );
+
+		static TArray<FVertexID> PolygonVertexIDs;
+		this->GetPolygonPerimeterVertices( PolygonID, /* Out */ PolygonVertexIDs );
+
+		// Create a new extruded polygon for the face
+		FPolygonID ExtrudedFrontPolygonID;	// Filled in below
+		{
+			static TArray<FVertexAndAttributes> NewFrontPolygonVertices;
+			NewFrontPolygonVertices.Reset( PolygonVertexIDs.Num() );
+			NewFrontPolygonVertices.SetNum( PolygonVertexIDs.Num(), false );
+
+			const TArray<FVertexInstanceID>& VertexInstanceIDs = GetMeshDescription()->GetPolygonVertexInstances( PolygonID );
+
+			// Map all of the polygon's vertex IDs to their extruded counterparts to create the new polygon perimeter
+			for( int32 PolygonVertexNumber = 0; PolygonVertexNumber < PolygonVertexIDs.Num(); ++PolygonVertexNumber )
+			{
+				const FVertexID VertexID = PolygonVertexIDs[ PolygonVertexNumber ];
+				const FVertexID* ExtrudedCopyVertexIDPtr = VertexIDToExtrudedCopy.Find( VertexID );
+				if( ExtrudedCopyVertexIDPtr != nullptr )
+				{
+					NewFrontPolygonVertices[ PolygonVertexNumber ].VertexID = VertexIDToExtrudedCopy[ VertexID ];
+				}
+				else
+				{
+					// We didn't need to extrude a new copy of this vertex (because it was part of a shared edge), so just connect the polygon to the original vertex
+					NewFrontPolygonVertices[ PolygonVertexNumber ].VertexID = VertexID;
+				}
+
+				// Copy vertex instance attributes from original polygon vertex to extruded polygon vertex
+				BackupAllAttributes( NewFrontPolygonVertices[ PolygonVertexNumber ].PolygonVertexAttributes, GetMeshDescription()->VertexInstanceAttributes(), VertexInstanceIDs[ PolygonVertexNumber ] );
+			}
+
+			{
+				static TArray<FPolygonToCreate> PolygonsToCreate;
+				PolygonsToCreate.Reset();
+
+				// Create the polygon
+				// @todo mesheditor perf: Ideally we support creating multiple polygons at once and batching up the work
+				FPolygonToCreate& PolygonToCreate = *new( PolygonsToCreate ) FPolygonToCreate();
+				PolygonToCreate.PolygonGroupID = PolygonGroupID;
+				PolygonToCreate.PolygonEdgeHardness = EPolygonEdgeHardness::AllEdgesHard;
+				PolygonToCreate.PerimeterVertices = NewFrontPolygonVertices;	// @todo mesheditor perf: Copying static array here, ideally allocations could be avoided
+				static TArray<FPolygonID> NewPolygonIDs;
+				NewPolygonIDs.Reset();
+				static TArray<FEdgeID> NewEdgeIDs;
+				NewEdgeIDs.Reset();
+				CreatePolygons( PolygonsToCreate, /* Out */ NewPolygonIDs, /* Out */ NewEdgeIDs );	// @todo mesheditor perf: Extra allocatons/copies: Ideally MoveTemp() here but we can't move a STATIC local!
+
+				ExtrudedFrontPolygonID = NewPolygonIDs[ 0 ];
+			}
+			AllNewPolygons.Add( ExtrudedFrontPolygonID );
+
+			// All of the border edges of the new polygon will be hard.  If it was a shared edge, then we'll just preserve whatever was
+			// originally going on with the internal edge.
+			{
+				const TEdgeAttributesRef<bool> EdgeHardnesses = GetMeshDescription()->EdgeAttributes().GetAttributesRef<bool>( MeshAttribute::Edge::IsHard );
+
+				const int32 NewPerimeterEdgeCount = this->GetPolygonPerimeterEdgeCount( ExtrudedFrontPolygonID );
+				check( NewPerimeterEdgeCount == GetPolygonPerimeterEdgeCount( PolygonID ) );	// New polygon should always have the same number of edges (in the same order) as the original!
+				for( int32 PerimeterEdgeNumber = 0; PerimeterEdgeNumber < NewPerimeterEdgeCount; ++PerimeterEdgeNumber )
+				{
+					bool bOriginalEdgeWindingIsReversedForPolygon;
+					const FEdgeID OriginalEdgeID = this->GetPolygonPerimeterEdge( PolygonID, PerimeterEdgeNumber, /* Out */ bOriginalEdgeWindingIsReversedForPolygon );
+					const bool bIsSharedEdge = bKeepNeighborsTogether && EdgeUsageCounts[ OriginalEdgeID ] > 1;
+
+					bool bEdgeWindingIsReversedForPolygon;
+					const FEdgeID EdgeID = this->GetPolygonPerimeterEdge( ExtrudedFrontPolygonID, PerimeterEdgeNumber, /* Out */ bEdgeWindingIsReversedForPolygon );
+
+					const bool bNewEdgeHardnessAttribute = bIsSharedEdge ? EdgeHardnesses[ OriginalEdgeID ] : true;
+
+					FAttributesForEdge& AttributesForEdge = *new( AttributesForEdges ) FAttributesForEdge();
+					AttributesForEdge.EdgeID = EdgeID;
+					AttributesForEdge.EdgeAttributes.Attributes.Emplace( MeshAttribute::Edge::IsHard, 0, FMeshElementAttributeValue( bNewEdgeHardnessAttribute ) );
+				}
+			}
+		}
+
+		OutNewExtrudedFrontPolygons.Add( ExtrudedFrontPolygonID );
+	}
+	check( NextAvailableExtrudedVertexIDNumber == ExtrudedVertexIDs.Num() );	// Make sure all of the vertices we created were actually used by new polygons
+
+																				// Update edge attributes in bulk
+	SetEdgesAttributes( AttributesForEdges );
+
+	// Update vertex attributes in bulk
+	SetVerticesAttributes( AttributesForVertices );
+
+	// Delete the original polygons
+	{
+		const bool bDeleteOrphanedEdges = true;
+		const bool bDeleteOrphanedVertices = true;
+		const bool bDeleteOrphanedVertexInstances = true;
+		const bool bDeleteEmptyPolygonGroups = false;
+		DeletePolygons( PolygonIDs, bDeleteOrphanedEdges, bDeleteOrphanedVertices, bDeleteOrphanedVertexInstances, bDeleteEmptyPolygonGroups );
+	}
+	#endif
+}
diff --git a/Engine/Source/Runtime/StaticMeshDescription/Public/StaticMeshAttributes.h b/Engine/Source/Runtime/StaticMeshDescription/Public/StaticMeshAttributes.h
new file mode 100644
--- /dev/null
+++ b/Engine/Source/Runtime/StaticMeshDescription/Public/StaticMeshAttributes.h
@@ -0,0 +1,120 @@
+// Copyright 1998-2019 Epic Games, Inc. All Rights Reserved.
+
+#pragma once
+
+#include "MeshAttributes.h"
+
+
+namespace MeshAttribute
+{
+	namespace Vertex
+	{
+		extern STATICMESHDESCRIPTION_API const FName CornerSharpness;
+	}
+
+	namespace VertexInstance
+	{
+		extern STATICMESHDESCRIPTION_API const FName TextureCoordinate;
+		extern STATICMESHDESCRIPTION_API const FName Normal;
+		extern STATICMESHDESCRIPTION_API const FName Tangent;
+		extern STATICMESHDESCRIPTION_API const FName BinormalSign;
+		extern STATICMESHDESCRIPTION_API const FName Color;
+	}
+
+	namespace Edge
+	{
+		extern STATICMESHDESCRIPTION_API const FName IsHard;
+		extern STATICMESHDESCRIPTION_API const FName IsUVSeam;
+		extern STATICMESHDESCRIPTION_API const FName CreaseSharpness;
+	}
+
+	namespace Polygon
+	{
+		extern STATICMESHDESCRIPTION_API const FName Normal;
+		extern STATICMESHDESCRIPTION_API const FName Tangent;
+		extern STATICMESHDESCRIPTION_API const FName Binormal;
+		extern STATICMESHDESCRIPTION_API const FName Center;
+	}
+
+	namespace PolygonGroup
+	{
+		extern STATICMESHDESCRIPTION_API const FName ImportedMaterialSlotName;
+		extern STATICMESHDESCRIPTION_API const FName EnableCollision;
+		extern STATICMESHDESCRIPTION_API const FName CastShadow;
+	}
+}
+
+
+class STATICMESHDESCRIPTION_API FStaticMeshAttributes : public FMeshAttributes
+{
+public:
+
+	explicit FStaticMeshAttributes(FMeshDescription& InMeshDescription)
+		: FMeshAttributes(InMeshDescription)
+	{}
+
+	virtual void Register() override;
+
+	TVertexAttributesRef<float> GetVertexCornerSharpnesses() { return MeshDescription.VertexAttributes().GetAttributesRef<float>(MeshAttribute::Vertex::CornerSharpness); }
+	TVertexAttributesConstRef<float> GetVertexCornerSharpnesses() const { return MeshDescription.VertexAttributes().GetAttributesRef<float>(MeshAttribute::Vertex::CornerSharpness); }
+
+	TVertexInstanceAttributesRef<FVector2D> GetVertexInstanceUVs() { return MeshDescription.VertexInstanceAttributes().GetAttributesRef<FVector2D>(MeshAttribute::VertexInstance::TextureCoordinate); }
+	TVertexInstanceAttributesConstRef<FVector2D> GetVertexInstanceUVs() const { return MeshDescription.VertexInstanceAttributes().GetAttributesRef<FVector2D>(MeshAttribute::VertexInstance::TextureCoordinate); }
+
+	TVertexInstanceAttributesRef<FVector> GetVertexInstanceNormals() { return MeshDescription.VertexInstanceAttributes().GetAttributesRef<FVector>(MeshAttribute::VertexInstance::Normal); }
+	TVertexInstanceAttributesConstRef<FVector> GetVertexInstanceNormals() const { return MeshDescription.VertexInstanceAttributes().GetAttributesRef<FVector>(MeshAttribute::VertexInstance::Normal); }
+
+	TVertexInstanceAttributesRef<FVector> GetVertexInstanceTangents() { return MeshDescription.VertexInstanceAttributes().GetAttributesRef<FVector>(MeshAttribute::VertexInstance::Tangent); }
+	TVertexInstanceAttributesConstRef<FVector> GetVertexInstanceTangents() const { return MeshDescription.VertexInstanceAttributes().GetAttributesRef<FVector>(MeshAttribute::VertexInstance::Tangent); }
+
+	TVertexInstanceAttributesRef<float> GetVertexInstanceBinormalSigns() { return MeshDescription.VertexInstanceAttributes().GetAttributesRef<float>(MeshAttribute::VertexInstance::BinormalSign); }
+	TVertexInstanceAttributesConstRef<float> GetVertexInstanceBinormalSigns() const { return MeshDescription.VertexInstanceAttributes().GetAttributesRef<float>(MeshAttribute::VertexInstance::BinormalSign); }
+
+	TVertexInstanceAttributesRef<FVector4> GetVertexInstanceColors() { return MeshDescription.VertexInstanceAttributes().GetAttributesRef<FVector4>(MeshAttribute::VertexInstance::Color); }
+	TVertexInstanceAttributesConstRef<FVector4> GetVertexInstanceColors() const { return MeshDescription.VertexInstanceAttributes().GetAttributesRef<FVector4>(MeshAttribute::VertexInstance::Color); }
+
+	TEdgeAttributesRef<bool> GetEdgeHardnesses() { return MeshDescription.EdgeAttributes().GetAttributesRef<bool>(MeshAttribute::Edge::IsHard); }
+	TEdgeAttributesConstRef<bool> GetEdgeHardnesses() const { return MeshDescription.EdgeAttributes().GetAttributesRef<bool>(MeshAttribute::Edge::IsHard); }
+
+	TEdgeAttributesRef<float> GetEdgeCreaseSharpnesses() { return MeshDescription.EdgeAttributes().GetAttributesRef<float>(MeshAttribute::Edge::CreaseSharpness); }
+	TEdgeAttributesConstRef<float> GetEdgeCreaseSharpnesses() const { return MeshDescription.EdgeAttributes().GetAttributesRef<float>(MeshAttribute::Edge::CreaseSharpness); }
+
+	TPolygonAttributesRef<FVector> GetPolygonNormals() { return MeshDescription.PolygonAttributes().GetAttributesRef<FVector>(MeshAttribute::Polygon::Normal); }
+	TPolygonAttributesConstRef<FVector> GetPolygonNormals() const { return MeshDescription.PolygonAttributes().GetAttributesRef<FVector>(MeshAttribute::Polygon::Normal); }
+
+	TPolygonAttributesRef<FVector> GetPolygonTangents() { return MeshDescription.PolygonAttributes().GetAttributesRef<FVector>(MeshAttribute::Polygon::Tangent); }
+	TPolygonAttributesConstRef<FVector> GetPolygonTangents() const { return MeshDescription.PolygonAttributes().GetAttributesRef<FVector>(MeshAttribute::Polygon::Tangent); }
+
+	TPolygonAttributesRef<FVector> GetPolygonBinormals() { return MeshDescription.PolygonAttributes().GetAttributesRef<FVector>(MeshAttribute::Polygon::Binormal); }
+	TPolygonAttributesConstRef<FVector> GetPolygonBinormals() const { return MeshDescription.PolygonAttributes().GetAttributesRef<FVector>(MeshAttribute::Polygon::Binormal); }
+
+	TPolygonAttributesRef<FVector> GetPolygonCenters() { return MeshDescription.PolygonAttributes().GetAttributesRef<FVector>(MeshAttribute::Polygon::Center); }
+	TPolygonAttributesConstRef<FVector> GetPolygonCenters() const { return MeshDescription.PolygonAttributes().GetAttributesRef<FVector>(MeshAttribute::Polygon::Center); }
+
+	TPolygonGroupAttributesRef<FName> GetPolygonGroupMaterialSlotNames() { return MeshDescription.PolygonGroupAttributes().GetAttributesRef<FName>(MeshAttribute::PolygonGroup::ImportedMaterialSlotName); }
+	TPolygonGroupAttributesConstRef<FName> GetPolygonGroupMaterialSlotNames() const { return MeshDescription.PolygonGroupAttributes().GetAttributesRef<FName>(MeshAttribute::PolygonGroup::ImportedMaterialSlotName); }
+};
+
+
+class FStaticMeshConstAttributes : public FMeshConstAttributes
+{
+public:
+
+	explicit FStaticMeshConstAttributes(const FMeshDescription& InMeshDescription)
+		: FMeshConstAttributes(InMeshDescription)
+	{}
+
+	TVertexAttributesConstRef<float> GetVertexCornerSharpnesses() const { return MeshDescription.VertexAttributes().GetAttributesRef<float>(MeshAttribute::Vertex::CornerSharpness); }
+	TVertexInstanceAttributesConstRef<FVector2D> GetVertexInstanceUVs() const { return MeshDescription.VertexInstanceAttributes().GetAttributesRef<FVector2D>(MeshAttribute::VertexInstance::TextureCoordinate); }
+	TVertexInstanceAttributesConstRef<FVector> GetVertexInstanceNormals() const { return MeshDescription.VertexInstanceAttributes().GetAttributesRef<FVector>(MeshAttribute::VertexInstance::Normal); }
+	TVertexInstanceAttributesConstRef<FVector> GetVertexInstanceTangents() const { return MeshDescription.VertexInstanceAttributes().GetAttributesRef<FVector>(MeshAttribute::VertexInstance::Tangent); }
+	TVertexInstanceAttributesConstRef<float> GetVertexInstanceBinormalSigns() const { return MeshDescription.VertexInstanceAttributes().GetAttributesRef<float>(MeshAttribute::VertexInstance::BinormalSign); }
+	TVertexInstanceAttributesConstRef<FVector4> GetVertexInstanceColors() const { return MeshDescription.VertexInstanceAttributes().GetAttributesRef<FVector4>(MeshAttribute::VertexInstance::Color); }
+	TEdgeAttributesConstRef<bool> GetEdgeHardnesses() const { return MeshDescription.EdgeAttributes().GetAttributesRef<bool>(MeshAttribute::Edge::IsHard); }
+	TEdgeAttributesConstRef<float> GetEdgeCreaseSharpnesses() const { return MeshDescription.EdgeAttributes().GetAttributesRef<float>(MeshAttribute::Edge::CreaseSharpness); }
+	TPolygonAttributesConstRef<FVector> GetPolygonNormals() const { return MeshDescription.PolygonAttributes().GetAttributesRef<FVector>(MeshAttribute::Polygon::Normal); }
+	TPolygonAttributesConstRef<FVector> GetPolygonTangents() const { return MeshDescription.PolygonAttributes().GetAttributesRef<FVector>(MeshAttribute::Polygon::Tangent); }
+	TPolygonAttributesConstRef<FVector> GetPolygonBinormals() const { return MeshDescription.PolygonAttributes().GetAttributesRef<FVector>(MeshAttribute::Polygon::Binormal); }
+	TPolygonAttributesConstRef<FVector> GetPolygonCenters() const { return MeshDescription.PolygonAttributes().GetAttributesRef<FVector>(MeshAttribute::Polygon::Center); }
+	TPolygonGroupAttributesConstRef<FName> GetPolygonGroupMaterialSlotNames() const { return MeshDescription.PolygonGroupAttributes().GetAttributesRef<FName>(MeshAttribute::PolygonGroup::ImportedMaterialSlotName); }
+};
diff --git a/Engine/Source/Runtime/StaticMeshDescription/Public/StaticMeshDescription.h b/Engine/Source/Runtime/StaticMeshDescription/Public/StaticMeshDescription.h
new file mode 100644
--- /dev/null
+++ b/Engine/Source/Runtime/StaticMeshDescription/Public/StaticMeshDescription.h
@@ -0,0 +1,80 @@
+// Copyright 1998-2019 Epic Games, Inc. All Rights Reserved.
+
+#pragma once
+
+#include "MeshDescriptionBase.h"
+#include "StaticMeshAttributes.h"
+#include "UObject/UObjectGlobals.h"
+#include "StaticMeshDescription.generated.h"
+
+class UMaterial;
+
+/**
+* A wrapper for MeshDescription, customized for static meshes
+*/
+UCLASS(BlueprintType)
+class STATICMESHDESCRIPTION_API UStaticMeshDescription : public UMeshDescriptionBase
+{
+public:
+	GENERATED_BODY()
+
+	/** Register attributes required by static mesh description */
+	virtual void RegisterAttributes() override;
+
+	virtual FStaticMeshAttributes& GetRequiredAttributes() override
+	{ 
+		return static_cast<FStaticMeshAttributes&>(*RequiredAttributes);
+	}
+
+	virtual const FStaticMeshAttributes& GetRequiredAttributes() const override
+	{
+		return static_cast<const FStaticMeshAttributes&>(*RequiredAttributes);
+	}
+
+	UFUNCTION(BlueprintPure, Category="MeshDescription")
+	FVector2D GetVertexInstanceUV(FVertexInstanceID VertexInstanceID, int32 UVIndex = 0) const;
+
+	UFUNCTION(BlueprintCallable, Category="MeshDescription")
+	void SetVertexInstanceUV(FVertexInstanceID VertexInstanceID, FVector2D UV, int32 UVIndex = 0);
+
+	UFUNCTION(BlueprintCallable, Category="MeshDescription")
+	void CreateCube(FVector Center, FVector HalfExtents, FPolygonGroupID PolygonGroup,
+					FPolygonID& PolygonID_PlusX,
+					FPolygonID& PolygonID_MinusX,
+					FPolygonID& PolygonID_PlusY,
+					FPolygonID& PolygonID_MinusY,
+					FPolygonID& PolygonID_PlusZ,
+					FPolygonID& PolygonID_MinusZ);
+
+	UFUNCTION(BlueprintCallable, Category="MeshDescription")
+	void SetPolygonGroupMaterialSlotName(FPolygonGroupID PolygonGroupID, const FName& SlotName);
+
+public:
+
+	TVertexAttributesRef<float> GetVertexCornerSharpnesses() { return GetRequiredAttributes().GetVertexCornerSharpnesses(); }
+	TVertexAttributesConstRef<float> GetVertexCornerSharpnesses() const { return GetRequiredAttributes().GetVertexCornerSharpnesses(); }
+
+	TVertexInstanceAttributesRef<FVector2D> GetVertexInstanceUVs() { return GetRequiredAttributes().GetVertexInstanceUVs(); }
+	TVertexInstanceAttributesConstRef<FVector2D> GetVertexInstanceUVs() const { return GetRequiredAttributes().GetVertexInstanceUVs(); }
+
+	TVertexInstanceAttributesRef<FVector> GetVertexInstanceNormals() { return GetRequiredAttributes().GetVertexInstanceNormals(); }
+	TVertexInstanceAttributesConstRef<FVector> GetVertexInstanceNormals() const { return GetRequiredAttributes().GetVertexInstanceNormals(); }
+
+	TVertexInstanceAttributesRef<FVector> GetVertexInstanceTangents() { return GetRequiredAttributes().GetVertexInstanceTangents(); }
+	TVertexInstanceAttributesConstRef<FVector> GetVertexInstanceTangents() const { return GetRequiredAttributes().GetVertexInstanceTangents(); }
+
+	TVertexInstanceAttributesRef<float> GetVertexInstanceBinormalSigns() { return GetRequiredAttributes().GetVertexInstanceBinormalSigns(); }
+	TVertexInstanceAttributesConstRef<float> GetVertexInstanceBinormalSigns() const { return GetRequiredAttributes().GetVertexInstanceBinormalSigns(); }
+
+	TVertexInstanceAttributesRef<FVector4> GetVertexInstanceColors() { return GetRequiredAttributes().GetVertexInstanceColors(); }
+	TVertexInstanceAttributesConstRef<FVector4> GetVertexInstanceColors() const { return GetRequiredAttributes().GetVertexInstanceColors(); }
+
+	TEdgeAttributesRef<bool> GetEdgeHardnesses() { return GetRequiredAttributes().GetEdgeHardnesses(); }
+	TEdgeAttributesConstRef<bool> GetEdgeHardnesses() const { return GetRequiredAttributes().GetEdgeHardnesses(); }
+
+	TEdgeAttributesRef<float> GetEdgeCreaseSharpnesses() { return GetRequiredAttributes().GetEdgeCreaseSharpnesses(); }
+	TEdgeAttributesConstRef<float> GetEdgeCreaseSharpnesses() const { return GetRequiredAttributes().GetEdgeCreaseSharpnesses(); }
+
+	TPolygonGroupAttributesRef<FName> GetPolygonGroupMaterialSlotNames() { return GetRequiredAttributes().GetPolygonGroupMaterialSlotNames(); }
+	TPolygonGroupAttributesConstRef<FName> GetPolygonGroupMaterialSlotNames() const { return GetRequiredAttributes().GetPolygonGroupMaterialSlotNames(); }
+};
diff --git a/Engine/Source/Runtime/StaticMeshDescription/Public/StaticMeshOperations.h b/Engine/Source/Runtime/StaticMeshDescription/Public/StaticMeshOperations.h
new file mode 100644
--- /dev/null
+++ b/Engine/Source/Runtime/StaticMeshDescription/Public/StaticMeshOperations.h
@@ -0,0 +1,45 @@
+// Copyright 1998-2019 Epic Games, Inc. All Rights Reserved.
+
+#pragma once
+
+#include "CoreMinimal.h"
+#include "MeshTypes.h"
+
+struct FMeshDescription;
+
+enum class EComputeNTBsFlags : uint32
+{
+	None = 0x00000000,	// No flags
+	Normals = 0x00000001, //Compute the normals
+	Tangents = 0x00000002, //Compute the tangents
+	WeightedNTBs = 0x00000004, //Use weight angle when computing NTBs to proportionally distribute the vertex instance contribution to the normal/tangent/binormal in a smooth group.    i.e. Weight solve the cylinder problem
+};
+ENUM_CLASS_FLAGS(EComputeNTBsFlags);
+
+
+class STATICMESHDESCRIPTION_API FStaticMeshOperations
+{
+public:
+
+	/** Set the polygon tangent and normal only for the specified polygonIDs */
+	static void ComputePolygonTangentsAndNormals(FMeshDescription& MeshDescription, TArrayView<const FPolygonID> PolygonIDs, float ComparisonThreshold = 0.0f);
+
+	/** Set the polygon tangent and normal for all polygons in the mesh description. */
+	static void ComputePolygonTangentsAndNormals(FMeshDescription& MeshDescription, float ComparisonThreshold = 0.0f);
+	
+	/** Set the vertex instance tangent and normal only for the specified VertexInstanceIDs */
+	static void ComputeTangentsAndNormals(FMeshDescription& MeshDescription, TArrayView<const FVertexInstanceID> VertexInstanceIDs, EComputeNTBsFlags ComputeNTBsOptions);
+
+	/** Set the vertex instance tangent and normal for all vertex instances in the mesh description. */
+	static void ComputeTangentsAndNormals(FMeshDescription& MeshDescription, EComputeNTBsFlags ComputeNTBsOptions);
+
+	/** Determine the edge hardnesses from existing normals */
+	static void DetermineEdgeHardnessesFromVertexInstanceNormals(FMeshDescription& MeshDescription, float Tolerance = KINDA_SMALL_NUMBER);
+
+	/** Extrude the given polygon IDs by the given distance */
+	static void ExtrudePolygons(FMeshDescription& MeshDescription, const TArray<FPolygonID>& PolygonIDs, float ExtrudeDistance, bool bKeepNeighborsTogether, TArray<FPolygonID>& OutNewExtrudedFrontPolygons);
+
+private:
+
+
+};
diff --git a/Engine/Source/Runtime/StaticMeshDescription/StaticMeshDescription.Build.cs b/Engine/Source/Runtime/StaticMeshDescription/StaticMeshDescription.Build.cs
new file mode 100644
--- /dev/null
+++ b/Engine/Source/Runtime/StaticMeshDescription/StaticMeshDescription.Build.cs
@@ -0,0 +1,22 @@
+// Copyright 1998-2019 Epic Games, Inc. All Rights Reserved.
+
+namespace UnrealBuildTool.Rules
+{
+	public class StaticMeshDescription : ModuleRules
+	{
+		public StaticMeshDescription(ReadOnlyTargetRules Target) : base(Target)
+		{
+            PrivateIncludePaths.Add("Runtime/StaticMeshDescription/Private");
+            PublicIncludePaths.Add("Runtime/StaticMeshDescription/Public");
+
+			PublicDependencyModuleNames.AddRange(
+				new string[]
+				{
+					"Core",
+					"CoreUObject",
+					"MeshDescription"
+				}
+			);
+		}
+	}
+}
