Removed deprecated functions.

diff --git a/Engine/Source/Runtime/StaticMeshDescription/Private/StaticMeshOperations.cpp b/Engine/Source/Runtime/StaticMeshDescription/Private/StaticMeshOperations.cpp
--- a/Engine/Source/Runtime/StaticMeshDescription/Private/StaticMeshOperations.cpp
+++ b/Engine/Source/Runtime/StaticMeshDescription/Private/StaticMeshOperations.cpp
@@ -433,346 +433,3 @@ void FStaticMeshOperations::DetermineEdgeHardnessesFromVertexInstanceNormals(FMe
 		EdgeHardnesses[EdgeID] = bEdgeIsHard;
 	}
 }
-
-
-void FStaticMeshOperations::ExtrudePolygons(FMeshDescription& MeshDescription, const TArray<FPolygonID>& PolygonIDs, float ExtrudeDistance, bool bKeepNeighborsTogether, TArray<FPolygonID>& OutNewExtrudedFrontPolygons)
-{
-	FStaticMeshAttributes Attributes(MeshDescription);
-	TVertexAttributesRef<FVector> VertexPositions = Attributes.GetVertexPositions();
-
-	OutNewExtrudedFrontPolygons.Reset();
-
-	// Convert our incoming polygon array to a TSet so we can lookup quickly to see which polygons in the mesh sare members of the set
-	TSet<FPolygonID> PolygonsSet;
-	PolygonsSet.Append(PolygonIDs);
-
-	TArray<FPolygonID> AllNewPolygons;
-
-//	static TArray<FAttributesForEdge> AttributesForEdges;
-//	AttributesForEdges.Reset();
-
-//	static TArray<FAttributesForVertex> AttributesForVertices;
-//	AttributesForVertices.Reset();
-
-	// First, let's figure out which of the polygons we were asked to extrude share edges or vertices.  We'll keep those edges intact!
-	TMap<FEdgeID, uint32> EdgeUsageCounts;	// Maps an edge ID to the number of times it is referenced by the incoming polygons
-	TSet<FVertexID> UniqueVertexIDs;
-
-	for (const FPolygonID PolygonID : PolygonIDs)
-	{
-		for (const FEdgeID EdgeID : MeshDescription.GetPolygonPerimeterEdges(PolygonID))
-		{
-			if (uint32* EdgeUsageCountPtr = EdgeUsageCounts.Find(EdgeID))
-			{
-				++(*EdgeUsageCountPtr);
-			}
-			else
-			{
-				EdgeUsageCounts.Add(EdgeID, 1);
-			}
-		}
-
-		UniqueVertexIDs.Append(MeshDescription.GetPolygonVertices(PolygonID));
-	}
-
-	int32 NumVerticesToCreate = UniqueVertexIDs.Num();
-
-	// Create new vertices for all of the extruded polygons
-	TArray<FVertexID> ExtrudedVertexIDs;
-	ExtrudedVertexIDs.Reserve(NumVerticesToCreate);
-	ExtrudedVertexIDs.Add(MeshDescription.CreateVertex());
-
-	int32 NextAvailableExtrudedVertexIDNumber = 0;
-
-	TMap<FVertexID, FVertexID> VertexIDToExtrudedCopy;
-	#if 0
-	for (int32 PassIndex = 0; PassIndex < 2; ++PassIndex)
-	{
-		// Extrude all of the shared edges first, then do the non-shared edges.  This is to make sure that a vertex doesn't get offset
-		// without taking into account all of the connected polygons in our set
-		const bool bIsExtrudingSharedEdges = (PassIndex == 0);
-
-		for (const FPolygonID PolygonID : PolygonIDs)
-		{
-			const FPolygonGroupID PolygonGroupID = MeshDescription.GetPolygonPolygonGroup(PolygonID);
-
-			if (!bKeepNeighborsTogether)
-			{
-				VertexIDToExtrudedCopy.Reset();
-			}
-
-			// Map all of the edge vertices to their new extruded counterpart
-			TArray<FVertexID> PolygonVertices = MeshDescription.GetPolygonVertices(PolygonID);
-			TArray<FEdgeID> PolygonPerimeterEdges = MeshDescription.GetPolygonPerimeterEdges(PolygonID);
-			for (const FEdgeID EdgeID : PolygonPerimeterEdges)
-			{
-				// @todo mesheditor perf: We can change GetPolygonPerimeterEdges() to have a version that returns whether winding is reversed or not, and avoid this call entirely.
-				// @todo mesheditor perf: O(log N^2) iteration here. For every edge, for every edge up to this index.  Need to clean this up. 
-				//		--> Also, there are quite a few places where we are stepping through edges in perimeter-order.  We need to have a nice way to walk that.
-				bool bEdgeWindingIsReversedForPolygon;
-				const FEdgeID EdgeID = this->GetPolygonPerimeterEdge( PolygonID, PerimeterEdgeNumber, /* Out */ bEdgeWindingIsReversedForPolygon );
-
-				const bool bIsSharedEdge = bKeepNeighborsTogether && EdgeUsageCounts[ EdgeID ] > 1;
-				if( bIsSharedEdge == bIsExtrudingSharedEdges )
-				{
-					FVertexID EdgeVertexIDs[ 2 ];
-					GetEdgeVertices( EdgeID, /* Out */ EdgeVertexIDs[ 0 ], /* Out */ EdgeVertexIDs[ 1 ] );
-					if( bEdgeWindingIsReversedForPolygon )
-					{
-						::Swap( EdgeVertexIDs[ 0 ], EdgeVertexIDs[ 1 ] );
-					}
-
-					if( !bIsSharedEdge )
-					{
-						// After extruding, all of the edges of the original polygon become hard edges
-						AttributesForEdges.Emplace();
-						FAttributesForEdge& AttributesForEdge = AttributesForEdges.Last();
-						AttributesForEdge.EdgeID = EdgeID;
-						AttributesForEdge.EdgeAttributes.Attributes.Emplace( MeshAttribute::Edge::IsHard, 0, FMeshElementAttributeValue( true ) );
-					}
-
-					FVertexID ExtrudedEdgeVertexIDs[ 2 ];
-					for( int32 EdgeVertexNumber = 0; EdgeVertexNumber < 2; ++EdgeVertexNumber )
-					{
-						const FVertexID EdgeVertexID = EdgeVertexIDs[ EdgeVertexNumber ];
-						FVertexID* ExtrudedEdgeVertexIDPtr = VertexIDToExtrudedCopy.Find( EdgeVertexID );
-
-						// @todo mesheditor extrude: Ideally we would detect whether the vertex that was already extruded came from a edge
-						// from a polygon that does not actually share an edge with any polygons this polygon shares an edge with.  This
-						// would avoid the problem where extruding two polygons that are connected only by a vertex are not extruded
-						// separately.
-						const bool bVertexIsSharedByAnEdgeOfAnotherSelectedPolygon = false;
-						if( ExtrudedEdgeVertexIDPtr != nullptr)// && !bVertexIsSharedByAnEdgeOfAnotherSelectedPolygon )
-						{
-							ExtrudedEdgeVertexIDs[ EdgeVertexNumber ] = *ExtrudedEdgeVertexIDPtr;
-						}
-						else
-						{
-							// Create a copy of this vertex for the extruded face
-							const FVertexID ExtrudedVertexID = ExtrudedVertexIDs[ NextAvailableExtrudedVertexIDNumber++ ];
-
-							ExtrudedEdgeVertexIDPtr = &VertexIDToExtrudedCopy.Add( EdgeVertexID, ExtrudedVertexID );
-							// Push the vertex out along the polygon's normal
-							const FVector OriginalVertexPosition = VertexPositions[ EdgeVertexID ];
-
-							FVector ExtrudedVertexPosition;
-							if( bIsSharedEdge )
-							{
-								// Get all of the polygons that share this edge that were part of the set of polygons passed in.  We'll
-								// generate an extrude direction that's the average of those polygon normals.
-								FVector ExtrudeDirection = FVector::ZeroVector;
-
-								static TArray<FPolygonID> ConnectedPolygonIDs;
-								GetVertexConnectedPolygons( EdgeVertexID, /* Out */ ConnectedPolygonIDs );
-
-								static TArray<FPolygonID> NeighborPolygonIDs;
-								NeighborPolygonIDs.Reset();
-								for( const FPolygonID ConnectedPolygonID : ConnectedPolygonIDs )
-								{
-									// We only care about polygons that are members of the set of polygons we were asked to extrude
-									if( PolygonsSet.Contains( ConnectedPolygonID ) )
-									{
-										NeighborPolygonIDs.Add( ConnectedPolygonID );
-
-										// We'll need this polygon's normal to figure out where to put the extruded copy of the polygon
-										const FVector NeighborPolygonNormal = ComputePolygonNormal( ConnectedPolygonID );
-										ExtrudeDirection += NeighborPolygonNormal;
-									}
-								}
-								ExtrudeDirection.Normalize();
-
-
-								// OK, we have the direction to extrude for this vertex.  Now we need to know how far to extrude.  We'll
-								// loop over all of the neighbor polygons to this vertex, and choose the closest intersection point with our
-								// vertex's extrude direction and the neighbor polygon's extruded plane
-								FVector ClosestIntersectionPointWithExtrudedPlanes;
-								float ClosestIntersectionDistanceSquared = TNumericLimits<float>::Max();
-
-								for( const FPolygonID NeighborPolygonID : NeighborPolygonIDs )
-								{
-									const FPlane NeighborPolygonPlane = ComputePolygonPlane( NeighborPolygonID );
-
-									// Push the plane out
-									const FPlane ExtrudedPlane = [NeighborPolygonPlane, ExtrudeDistance]
-									{ 
-										FPlane NewPlane = NeighborPolygonPlane;
-										NewPlane.W += ExtrudeDistance; 
-										return NewPlane; 
-									}();
-
-									// Is this the closest intersection point so far?
-									const FVector IntersectionPointWithExtrudedPlane = FMath::RayPlaneIntersection( OriginalVertexPosition, ExtrudeDirection, ExtrudedPlane );
-									const float IntersectionDistanceSquared = FVector::DistSquared( OriginalVertexPosition, IntersectionPointWithExtrudedPlane );
-									if( IntersectionDistanceSquared < ClosestIntersectionDistanceSquared )
-									{
-										ClosestIntersectionPointWithExtrudedPlanes = IntersectionPointWithExtrudedPlane;
-										ClosestIntersectionDistanceSquared = IntersectionDistanceSquared;
-									}
-								}
-
-								ExtrudedVertexPosition = ClosestIntersectionPointWithExtrudedPlanes;
-							}
-							else
-							{
-								// We'll need this polygon's normal to figure out where to put the extruded copy of the polygon
-								const FVector PolygonNormal = ComputePolygonNormal( PolygonID );
-								ExtrudedVertexPosition = OriginalVertexPosition + ExtrudeDistance * PolygonNormal;
-							}
-
-							// Fill in the vertex
-							AttributesForVertices.Emplace();
-							FAttributesForVertex& AttributesForVertex = AttributesForVertices.Last();
-							AttributesForVertex.VertexID = ExtrudedVertexID;
-							AttributesForVertex.VertexAttributes.Attributes.Emplace( MeshAttribute::Vertex::Position, 0, FMeshElementAttributeValue( ExtrudedVertexPosition ) );
-						}
-						ExtrudedEdgeVertexIDs[ EdgeVertexNumber ] = *ExtrudedEdgeVertexIDPtr;
-					}
-
-					if( !bIsSharedEdge )
-					{
-						static TArray<FVertexAndAttributes> NewSidePolygonVertices;
-						NewSidePolygonVertices.Reset( 4 );
-						NewSidePolygonVertices.SetNum( 4, false );	// Always four edges in an extruded face
-
-						NewSidePolygonVertices[ 0 ].VertexID = EdgeVertexIDs[ 1 ];
-						NewSidePolygonVertices[ 1 ].VertexID = EdgeVertexIDs[ 0 ];
-						NewSidePolygonVertices[ 2 ].VertexID = ExtrudedEdgeVertexIDs[ 0 ];
-						NewSidePolygonVertices[ 3 ].VertexID = ExtrudedEdgeVertexIDs[ 1 ];
-
-						// Get vertex instance IDs on this polygon corresponding to the edge start/end vertices
-						const FVertexInstanceID EdgeVertexInstanceID0 = GetMeshDescription()->GetVertexInstanceForPolygonVertex( PolygonID, EdgeVertexIDs[ 0 ] );
-						const FVertexInstanceID EdgeVertexInstanceID1 = GetMeshDescription()->GetVertexInstanceForPolygonVertex( PolygonID, EdgeVertexIDs[ 1 ] );
-
-						BackupAllAttributes( NewSidePolygonVertices[ 0 ].PolygonVertexAttributes, GetMeshDescription()->VertexInstanceAttributes(), EdgeVertexInstanceID1 );
-						BackupAllAttributes( NewSidePolygonVertices[ 1 ].PolygonVertexAttributes, GetMeshDescription()->VertexInstanceAttributes(), EdgeVertexInstanceID0 );
-						BackupAllAttributes( NewSidePolygonVertices[ 2 ].PolygonVertexAttributes, GetMeshDescription()->VertexInstanceAttributes(), EdgeVertexInstanceID0 );
-						BackupAllAttributes( NewSidePolygonVertices[ 3 ].PolygonVertexAttributes, GetMeshDescription()->VertexInstanceAttributes(), EdgeVertexInstanceID1 );
-
-						FPolygonID NewSidePolygonID;	// Filled in below
-						{
-							static TArray<FPolygonToCreate> PolygonsToCreate;
-							PolygonsToCreate.Reset();
-
-							// Create the polygon
-							// @todo mesheditor perf: Ideally we support creating multiple polygons at once and batching up the work
-							FPolygonToCreate& PolygonToCreate = *new( PolygonsToCreate ) FPolygonToCreate();
-							PolygonToCreate.PolygonGroupID = PolygonGroupID;
-							PolygonToCreate.PerimeterVertices = NewSidePolygonVertices;	// @todo mesheditor perf: Copying static array here, ideally allocations could be avoided
-							PolygonToCreate.PolygonEdgeHardness = EPolygonEdgeHardness::AllEdgesHard;
-
-							static TArray<FPolygonID> NewPolygonIDs;
-							NewPolygonIDs.Reset();
-							static TArray<FEdgeID> NewEdgeIDs;
-							NewEdgeIDs.Reset();
-							CreatePolygons( PolygonsToCreate, /* Out */ NewPolygonIDs, /* Out */ NewEdgeIDs );	// @todo mesheditor perf: Extra allocatons/copies: Ideally MoveTemp() here but we can't move a STATIC local!
-
-							NewSidePolygonID = NewPolygonIDs[ 0 ];
-						}
-						AllNewPolygons.Add( NewSidePolygonID );
-					}
-				}
-			}
-		}
-	}
-
-	for( int32 PolygonIter = 0; PolygonIter < PolygonIDs.Num(); ++PolygonIter )
-	{
-		const FPolygonID PolygonID = PolygonIDs[ PolygonIter ];
-		const FPolygonGroupID PolygonGroupID = GetGroupForPolygon( PolygonID );
-
-		static TArray<FVertexID> PolygonVertexIDs;
-		this->GetPolygonPerimeterVertices( PolygonID, /* Out */ PolygonVertexIDs );
-
-		// Create a new extruded polygon for the face
-		FPolygonID ExtrudedFrontPolygonID;	// Filled in below
-		{
-			static TArray<FVertexAndAttributes> NewFrontPolygonVertices;
-			NewFrontPolygonVertices.Reset( PolygonVertexIDs.Num() );
-			NewFrontPolygonVertices.SetNum( PolygonVertexIDs.Num(), false );
-
-			const TArray<FVertexInstanceID>& VertexInstanceIDs = GetMeshDescription()->GetPolygonVertexInstances( PolygonID );
-
-			// Map all of the polygon's vertex IDs to their extruded counterparts to create the new polygon perimeter
-			for( int32 PolygonVertexNumber = 0; PolygonVertexNumber < PolygonVertexIDs.Num(); ++PolygonVertexNumber )
-			{
-				const FVertexID VertexID = PolygonVertexIDs[ PolygonVertexNumber ];
-				const FVertexID* ExtrudedCopyVertexIDPtr = VertexIDToExtrudedCopy.Find( VertexID );
-				if( ExtrudedCopyVertexIDPtr != nullptr )
-				{
-					NewFrontPolygonVertices[ PolygonVertexNumber ].VertexID = VertexIDToExtrudedCopy[ VertexID ];
-				}
-				else
-				{
-					// We didn't need to extrude a new copy of this vertex (because it was part of a shared edge), so just connect the polygon to the original vertex
-					NewFrontPolygonVertices[ PolygonVertexNumber ].VertexID = VertexID;
-				}
-
-				// Copy vertex instance attributes from original polygon vertex to extruded polygon vertex
-				BackupAllAttributes( NewFrontPolygonVertices[ PolygonVertexNumber ].PolygonVertexAttributes, GetMeshDescription()->VertexInstanceAttributes(), VertexInstanceIDs[ PolygonVertexNumber ] );
-			}
-
-			{
-				static TArray<FPolygonToCreate> PolygonsToCreate;
-				PolygonsToCreate.Reset();
-
-				// Create the polygon
-				// @todo mesheditor perf: Ideally we support creating multiple polygons at once and batching up the work
-				FPolygonToCreate& PolygonToCreate = *new( PolygonsToCreate ) FPolygonToCreate();
-				PolygonToCreate.PolygonGroupID = PolygonGroupID;
-				PolygonToCreate.PolygonEdgeHardness = EPolygonEdgeHardness::AllEdgesHard;
-				PolygonToCreate.PerimeterVertices = NewFrontPolygonVertices;	// @todo mesheditor perf: Copying static array here, ideally allocations could be avoided
-				static TArray<FPolygonID> NewPolygonIDs;
-				NewPolygonIDs.Reset();
-				static TArray<FEdgeID> NewEdgeIDs;
-				NewEdgeIDs.Reset();
-				CreatePolygons( PolygonsToCreate, /* Out */ NewPolygonIDs, /* Out */ NewEdgeIDs );	// @todo mesheditor perf: Extra allocatons/copies: Ideally MoveTemp() here but we can't move a STATIC local!
-
-				ExtrudedFrontPolygonID = NewPolygonIDs[ 0 ];
-			}
-			AllNewPolygons.Add( ExtrudedFrontPolygonID );
-
-			// All of the border edges of the new polygon will be hard.  If it was a shared edge, then we'll just preserve whatever was
-			// originally going on with the internal edge.
-			{
-				const TEdgeAttributesRef<bool> EdgeHardnesses = GetMeshDescription()->EdgeAttributes().GetAttributesRef<bool>( MeshAttribute::Edge::IsHard );
-
-				const int32 NewPerimeterEdgeCount = this->GetPolygonPerimeterEdgeCount( ExtrudedFrontPolygonID );
-				check( NewPerimeterEdgeCount == GetPolygonPerimeterEdgeCount( PolygonID ) );	// New polygon should always have the same number of edges (in the same order) as the original!
-				for( int32 PerimeterEdgeNumber = 0; PerimeterEdgeNumber < NewPerimeterEdgeCount; ++PerimeterEdgeNumber )
-				{
-					bool bOriginalEdgeWindingIsReversedForPolygon;
-					const FEdgeID OriginalEdgeID = this->GetPolygonPerimeterEdge( PolygonID, PerimeterEdgeNumber, /* Out */ bOriginalEdgeWindingIsReversedForPolygon );
-					const bool bIsSharedEdge = bKeepNeighborsTogether && EdgeUsageCounts[ OriginalEdgeID ] > 1;
-
-					bool bEdgeWindingIsReversedForPolygon;
-					const FEdgeID EdgeID = this->GetPolygonPerimeterEdge( ExtrudedFrontPolygonID, PerimeterEdgeNumber, /* Out */ bEdgeWindingIsReversedForPolygon );
-
-					const bool bNewEdgeHardnessAttribute = bIsSharedEdge ? EdgeHardnesses[ OriginalEdgeID ] : true;
-
-					FAttributesForEdge& AttributesForEdge = *new( AttributesForEdges ) FAttributesForEdge();
-					AttributesForEdge.EdgeID = EdgeID;
-					AttributesForEdge.EdgeAttributes.Attributes.Emplace( MeshAttribute::Edge::IsHard, 0, FMeshElementAttributeValue( bNewEdgeHardnessAttribute ) );
-				}
-			}
-		}
-
-		OutNewExtrudedFrontPolygons.Add( ExtrudedFrontPolygonID );
-	}
-	check( NextAvailableExtrudedVertexIDNumber == ExtrudedVertexIDs.Num() );	// Make sure all of the vertices we created were actually used by new polygons
-
-																				// Update edge attributes in bulk
-	SetEdgesAttributes( AttributesForEdges );
-
-	// Update vertex attributes in bulk
-	SetVerticesAttributes( AttributesForVertices );
-
-	// Delete the original polygons
-	{
-		const bool bDeleteOrphanedEdges = true;
-		const bool bDeleteOrphanedVertices = true;
-		const bool bDeleteOrphanedVertexInstances = true;
-		const bool bDeleteEmptyPolygonGroups = false;
-		DeletePolygons( PolygonIDs, bDeleteOrphanedEdges, bDeleteOrphanedVertices, bDeleteOrphanedVertexInstances, bDeleteEmptyPolygonGroups );
-	}
-	#endif
-}
diff --git a/Engine/Source/Runtime/StaticMeshDescription/Public/StaticMeshOperations.h b/Engine/Source/Runtime/StaticMeshDescription/Public/StaticMeshOperations.h
--- a/Engine/Source/Runtime/StaticMeshDescription/Public/StaticMeshOperations.h
+++ b/Engine/Source/Runtime/StaticMeshDescription/Public/StaticMeshOperations.h
@@ -35,11 +35,4 @@ public:
 
 	/** Determine the edge hardnesses from existing normals */
 	static void DetermineEdgeHardnessesFromVertexInstanceNormals(FMeshDescription& MeshDescription, float Tolerance = KINDA_SMALL_NUMBER);
-
-	/** Extrude the given polygon IDs by the given distance */
-	static void ExtrudePolygons(FMeshDescription& MeshDescription, const TArray<FPolygonID>& PolygonIDs, float ExtrudeDistance, bool bKeepNeighborsTogether, TArray<FPolygonID>& OutNewExtrudedFrontPolygons);
-
-private:
-
-
 };
