Fixed issue in which generating normals/tangents in a generated static mesh could lead to a crash.

diff --git a/Engine/Source/Runtime/StaticMeshDescription/Private/StaticMeshAttributes.cpp b/Engine/Source/Runtime/StaticMeshDescription/Private/StaticMeshAttributes.cpp
--- a/Engine/Source/Runtime/StaticMeshDescription/Private/StaticMeshAttributes.cpp
+++ b/Engine/Source/Runtime/StaticMeshDescription/Private/StaticMeshAttributes.cpp
@@ -52,3 +52,11 @@ void FStaticMeshAttributes::Register()
 	// Call super class
 	FMeshAttributes::Register();
 }
+
+
+void FStaticMeshAttributes::RegisterPolygonNormalAndTangentAttributes()
+{
+	MeshDescription.PolygonAttributes().RegisterAttribute<FVector>(MeshAttribute::Polygon::Normal, 1, FVector::ZeroVector, EMeshAttributeFlags::Transient);
+	MeshDescription.PolygonAttributes().RegisterAttribute<FVector>(MeshAttribute::Polygon::Tangent, 1, FVector::ZeroVector, EMeshAttributeFlags::Transient);
+	MeshDescription.PolygonAttributes().RegisterAttribute<FVector>(MeshAttribute::Polygon::Binormal, 1, FVector::ZeroVector, EMeshAttributeFlags::Transient);
+}
diff --git a/Engine/Source/Runtime/StaticMeshDescription/Private/StaticMeshOperations.cpp b/Engine/Source/Runtime/StaticMeshDescription/Private/StaticMeshOperations.cpp
--- a/Engine/Source/Runtime/StaticMeshDescription/Private/StaticMeshOperations.cpp
+++ b/Engine/Source/Runtime/StaticMeshDescription/Private/StaticMeshOperations.cpp
@@ -270,36 +270,16 @@ static void GetTangentsAndNormals(FMeshDescription& MeshDescription,
 }
 
 
-
-void FStaticMeshOperations::ComputePolygonTangentsAndNormals(FMeshDescription& MeshDescription, TArrayView<const FPolygonID> PolygonIDs, float ComparisonThreshold)
+static void RemovePolygonsAndCompact(FMeshDescription& MeshDescription, const TArray<FPolygonID>& PolygonsToRemove)
 {
-	FStaticMeshAttributes Attributes(MeshDescription);
-	TVertexAttributesRef<const FVector> VertexPositions = Attributes.GetVertexPositions();
-	TVertexInstanceAttributesRef<const FVector2D> VertexUVs = Attributes.GetVertexInstanceUVs();
-	TPolygonAttributesRef<FVector> PolygonNormals = Attributes.GetPolygonNormals();
-	TPolygonAttributesRef<FVector> PolygonTangents = Attributes.GetPolygonTangents();
-	TPolygonAttributesRef<FVector> PolygonBinormals = Attributes.GetPolygonBinormals();
-	TPolygonAttributesRef<FVector> PolygonCenters = Attributes.GetPolygonCenters();
-
-	TArray<FPolygonID> DegeneratePolygonIDs;
-	for (const FPolygonID PolygonID : PolygonIDs)
-	{
-		if (!GetPolygonTangentsAndNormals(MeshDescription, PolygonID, ComparisonThreshold, VertexPositions, VertexUVs, PolygonNormals, PolygonTangents, PolygonBinormals, PolygonCenters))
-		{
-			DegeneratePolygonIDs.Add(PolygonID);
-		}
-	}
-
-	// Remove degenerated polygons
-	// Delete the degenerated polygons. The array is filled only if the remove degenerated option is turned on.
-	if (DegeneratePolygonIDs.Num() > 0)
+	if (PolygonsToRemove.Num() > 0)
 	{
 		TArray<FEdgeID> OrphanedEdges;
 		TArray<FVertexInstanceID> OrphanedVertexInstances;
 		TArray<FPolygonGroupID> OrphanedPolygonGroups;
 		TArray<FVertexID> OrphanedVertices;
 
-		for (FPolygonID PolygonID : DegeneratePolygonIDs)
+		for (FPolygonID PolygonID : PolygonsToRemove)
 		{
 			MeshDescription.DeletePolygon(PolygonID, &OrphanedEdges, &OrphanedVertexInstances, &OrphanedPolygonGroups);
 		}
@@ -325,22 +305,60 @@ void FStaticMeshOperations::ComputePolygonTangentsAndNormals(FMeshDescription& M
 	}
 }
 
+
+void FStaticMeshOperations::ComputePolygonTangentsAndNormals(FMeshDescription& MeshDescription, TArrayView<const FPolygonID> PolygonIDs, float ComparisonThreshold)
+{
+	FStaticMeshAttributes Attributes(MeshDescription);
+	Attributes.RegisterPolygonNormalAndTangentAttributes();
+	TVertexAttributesRef<const FVector> VertexPositions = Attributes.GetVertexPositions();
+	TVertexInstanceAttributesRef<const FVector2D> VertexUVs = Attributes.GetVertexInstanceUVs();
+	TPolygonAttributesRef<FVector> PolygonNormals = Attributes.GetPolygonNormals();
+	TPolygonAttributesRef<FVector> PolygonTangents = Attributes.GetPolygonTangents();
+	TPolygonAttributesRef<FVector> PolygonBinormals = Attributes.GetPolygonBinormals();
+	TPolygonAttributesRef<FVector> PolygonCenters = Attributes.GetPolygonCenters();
+
+	TArray<FPolygonID> DegeneratePolygonIDs;
+	for (const FPolygonID PolygonID : PolygonIDs)
+	{
+		if (!GetPolygonTangentsAndNormals(MeshDescription, PolygonID, ComparisonThreshold, VertexPositions, VertexUVs, PolygonNormals, PolygonTangents, PolygonBinormals, PolygonCenters))
+		{
+			DegeneratePolygonIDs.Add(PolygonID);
+		}
+	}
+
+	// Remove degenerated polygons
+	RemovePolygonsAndCompact(MeshDescription, DegeneratePolygonIDs);
+}
+
 void FStaticMeshOperations::ComputePolygonTangentsAndNormals(FMeshDescription& MeshDescription, float ComparisonThreshold)
 {
-	TArray<FPolygonID> PolygonsToComputeNTBs;
-	PolygonsToComputeNTBs.Reserve(MeshDescription.Polygons().Num());
+	FStaticMeshAttributes Attributes(MeshDescription);
+	Attributes.RegisterPolygonNormalAndTangentAttributes();
+	TVertexAttributesRef<const FVector> VertexPositions = Attributes.GetVertexPositions();
+	TVertexInstanceAttributesRef<const FVector2D> VertexUVs = Attributes.GetVertexInstanceUVs();
+	TPolygonAttributesRef<FVector> PolygonNormals = Attributes.GetPolygonNormals();
+	TPolygonAttributesRef<FVector> PolygonTangents = Attributes.GetPolygonTangents();
+	TPolygonAttributesRef<FVector> PolygonBinormals = Attributes.GetPolygonBinormals();
+	TPolygonAttributesRef<FVector> PolygonCenters = Attributes.GetPolygonCenters();
 
+	TArray<FPolygonID> DegeneratePolygonIDs;
 	for (const FPolygonID PolygonID : MeshDescription.Polygons().GetElementIDs())
 	{
-		PolygonsToComputeNTBs.Add(PolygonID);
+		if (!GetPolygonTangentsAndNormals(MeshDescription, PolygonID, ComparisonThreshold, VertexPositions, VertexUVs, PolygonNormals, PolygonTangents, PolygonBinormals, PolygonCenters))
+		{
+			DegeneratePolygonIDs.Add(PolygonID);
+		}
 	}
 
-	ComputePolygonTangentsAndNormals(MeshDescription, PolygonsToComputeNTBs, ComparisonThreshold);
+	// Remove degenerated polygons
+	RemovePolygonsAndCompact(MeshDescription, DegeneratePolygonIDs);
 }
 
+
 void FStaticMeshOperations::ComputeTangentsAndNormals(FMeshDescription& MeshDescription, TArrayView<const FVertexInstanceID> VertexInstanceIDs, EComputeNTBsFlags ComputeNTBsOptions)
 {
 	FStaticMeshAttributes Attributes(MeshDescription);
+	Attributes.RegisterPolygonNormalAndTangentAttributes();
 	TPolygonAttributesRef<const FVector> PolygonNormals = Attributes.GetPolygonNormals();
 	TPolygonAttributesRef<const FVector> PolygonTangents = Attributes.GetPolygonTangents();
 	TPolygonAttributesRef<const FVector> PolygonBinormals = Attributes.GetPolygonBinormals();
@@ -354,9 +372,11 @@ void FStaticMeshOperations::ComputeTangentsAndNormals(FMeshDescription& MeshDesc
 	}
 }
 
+
 void FStaticMeshOperations::ComputeTangentsAndNormals(FMeshDescription& MeshDescription, EComputeNTBsFlags ComputeNTBsOptions)
 {
 	FStaticMeshAttributes Attributes(MeshDescription);
+	Attributes.RegisterPolygonNormalAndTangentAttributes();
 	TPolygonAttributesRef<const FVector> PolygonNormals = Attributes.GetPolygonNormals();
 	TPolygonAttributesRef<const FVector> PolygonTangents = Attributes.GetPolygonTangents();
 	TPolygonAttributesRef<const FVector> PolygonBinormals = Attributes.GetPolygonBinormals();
diff --git a/Engine/Source/Runtime/StaticMeshDescription/Public/StaticMeshAttributes.h b/Engine/Source/Runtime/StaticMeshDescription/Public/StaticMeshAttributes.h
--- a/Engine/Source/Runtime/StaticMeshDescription/Public/StaticMeshAttributes.h
+++ b/Engine/Source/Runtime/StaticMeshDescription/Public/StaticMeshAttributes.h
@@ -55,6 +55,8 @@ public:
 
 	virtual void Register() override;
 
+	void RegisterPolygonNormalAndTangentAttributes();
+
 	TVertexAttributesRef<float> GetVertexCornerSharpnesses() { return MeshDescription.VertexAttributes().GetAttributesRef<float>(MeshAttribute::Vertex::CornerSharpness); }
 	TVertexAttributesConstRef<float> GetVertexCornerSharpnesses() const { return MeshDescription.VertexAttributes().GetAttributesRef<float>(MeshAttribute::Vertex::CornerSharpness); }
 
