Added StaticMeshDescription module for defining attributes and operations required in editor and at runtime for MeshDescriptions for static meshes.
Provided a path for building MeshDescriptions for static meshes at runtime, and supporting API.
Clarified use of attributes for static meshes; now some attributes may be considered optional.
Added FStaticMeshAttributes, used for fetching and registering all the attributes required by a static mesh.
Added Blueprint interface for building meshes.
Deprecated a number of existing APIs:
 - FMeshDescription methods which referred explicitly to static mesh specific attributes
 - StaticMesh attribute registration has been moved to FStaticMeshAttributes

diff --git a/Engine/Plugins/Experimental/MeshModelingToolset/Source/ModelingComponents/Private/OctreeDynamicMeshComponent.cpp b/Engine/Plugins/Experimental/MeshModelingToolset/Source/ModelingComponents/Private/OctreeDynamicMeshComponent.cpp
--- a/Engine/Plugins/Experimental/MeshModelingToolset/Source/ModelingComponents/Private/OctreeDynamicMeshComponent.cpp
+++ b/Engine/Plugins/Experimental/MeshModelingToolset/Source/ModelingComponents/Private/OctreeDynamicMeshComponent.cpp
@@ -15,9 +15,7 @@
 #include "EngineGlobals.h"
 #include "Engine/Engine.h"
 #include "StaticMeshResources.h"
-
-#include "MeshDescription.h"
-#include "MeshAttributes.h"
+#include "StaticMeshAttributes.h"
 
 #include "DynamicMeshAttributeSet.h"
 #include "MeshNormals.h"
diff --git a/Engine/Source/Developer/MeshUtilities/Private/MeshUtilities.cpp b/Engine/Source/Developer/MeshUtilities/Private/MeshUtilities.cpp
--- a/Engine/Source/Developer/MeshUtilities/Private/MeshUtilities.cpp
+++ b/Engine/Source/Developer/MeshUtilities/Private/MeshUtilities.cpp
@@ -94,7 +94,7 @@
 #include "OverlappingCorners.h"
 #include "MeshUtilitiesCommon.h"
 
-#include "MeshDescription.h"
+#include "StaticMeshAttributes.h"
 #include "MeshDescriptionOperations.h"
 
 #if WITH_EDITOR
@@ -2632,10 +2632,13 @@ public:
 				FRawMesh& DestMesh = LODMeshes[NumValidLODs];
 				FOverlappingCorners& InOverlappingCorners = LODOverlappingCorners[ReductionSettings.BaseLODModel];
 				FOverlappingCorners& DestOverlappingCorners = LODOverlappingCorners[NumValidLODs];
+
 				FMeshDescription SrcMeshdescription;
-				UStaticMesh::RegisterMeshAttributes(SrcMeshdescription);
+				FStaticMeshAttributes(SrcMeshdescription).Register();
+
 				FMeshDescription DestMeshdescription;
-				UStaticMesh::RegisterMeshAttributes(DestMeshdescription);
+				FStaticMeshAttributes(DestMeshdescription).Register();
+
 				TMap<int32, FName> FromMaterialMap;
 				FMeshDescriptionOperations::ConvertFromRawMesh(InMesh, SrcMeshdescription, FromMaterialMap);
 				MeshReduction->ReduceMeshDescription(DestMeshdescription, LODMaxDeviation[NumValidLODs], SrcMeshdescription, InOverlappingCorners, ReductionSettings);
diff --git a/Engine/Source/Editor/MeshPaint/Private/MeshPaintHelpers.cpp b/Engine/Source/Editor/MeshPaint/Private/MeshPaintHelpers.cpp
--- a/Engine/Source/Editor/MeshPaint/Private/MeshPaintHelpers.cpp
+++ b/Engine/Source/Editor/MeshPaint/Private/MeshPaintHelpers.cpp
@@ -14,10 +14,7 @@
 #include "Engine/SkeletalMesh.h"
 #include "Engine/Texture2D.h"
 #include "StaticMeshResources.h"
-
-#include "MeshDescription.h"
-#include "MeshAttributes.h"
-#include "MeshAttributeArray.h"
+#include "StaticMeshAttributes.h"
 
 #include "Rendering/SkeletalMeshRenderData.h"
 
@@ -112,22 +109,21 @@ bool MeshPaintHelpers::PropagateColorsToRawMesh(UStaticMesh* StaticMesh, int32 L
 	if (RenderData.WedgeMap.Num() > 0 && ColorVertexBuffer.GetNumVertices() == RenderModel.GetNumVertices())
 	{
 		// Use the wedge map if it is available as it is lossless.
-		FMeshDescription* RawMeshPtr = StaticMesh->GetMeshDescription(LODIndex);
+		FMeshDescription* MeshDescription = StaticMesh->GetMeshDescription(LODIndex);
 
-		if (RawMeshPtr == nullptr)
+		if (MeshDescription == nullptr)
 		{
 			//Cannot propagate to a generated LOD, the generated LOD use the source LOD vertex painting. 
 			return false;
 		}
 
-		FMeshDescription& RawMesh = *RawMeshPtr;
-		int32 NumWedges = RawMesh.VertexInstances().Num();
+		FStaticMeshAttributes Attributes(*MeshDescription);
+		int32 NumWedges = MeshDescription->VertexInstances().Num();
 		if (RenderData.WedgeMap.Num() == NumWedges)
 		{
-			FStaticMeshDescriptionAttributeGetter AttributeGetter(&RawMesh);
-			TVertexInstanceAttributesRef<FVector4> Colors = AttributeGetter.GetColors();
+			TVertexInstanceAttributesRef<FVector4> Colors = Attributes.GetVertexInstanceColors();
 			int32 VertexInstanceIndex = 0;
-			for (const FVertexInstanceID VertexInstanceID : RawMesh.VertexInstances().GetElementIDs())
+			for (const FVertexInstanceID VertexInstanceID : MeshDescription->VertexInstances().GetElementIDs())
 			{
 				FLinearColor WedgeColor = FLinearColor::White;
 				int32 Index = RenderData.WedgeMap[VertexInstanceIndex];
@@ -144,24 +140,25 @@ bool MeshPaintHelpers::PropagateColorsToRawMesh(UStaticMesh* StaticMesh, int32 L
 	}
 	else
 	{
-		FMeshDescription* SrcModelMeshDescription = StaticMesh->GetMeshDescription(LODIndex);
+		FMeshDescription* MeshDescription = StaticMesh->GetMeshDescription(LODIndex);
 		// If there's no raw mesh data, don't try to do any fixup here
-		if (SrcModelMeshDescription == nullptr || ComponentLODInfo.OverrideVertexColors == nullptr)
+		if (MeshDescription == nullptr || ComponentLODInfo.OverrideVertexColors == nullptr)
 		{
 			return false;
 		}
 
+		FStaticMeshAttributes Attributes(*MeshDescription);
+
 		// Fall back to mapping based on position.
-		FStaticMeshDescriptionAttributeGetter AttributeGetter(SrcModelMeshDescription);
-		TVertexAttributesConstRef<FVector> VertexPositions = AttributeGetter.GetPositionsConst();
-		TVertexInstanceAttributesRef<FVector4> Colors = AttributeGetter.GetColors();
+		TVertexAttributesConstRef<FVector> VertexPositions = Attributes.GetVertexPositions();
+		TVertexInstanceAttributesRef<FVector4> Colors = Attributes.GetVertexInstanceColors();
 		TArray<FColor> NewVertexColors;
 		FPositionVertexBuffer TempPositionVertexBuffer;
-		int32 NumVertex = SrcModelMeshDescription->Vertices().Num();
+		int32 NumVertex = MeshDescription->Vertices().Num();
 		TArray<FVector> VertexPositionsDup;
 		VertexPositionsDup.AddZeroed(NumVertex);
 		int32 VertexIndex = 0;
-		for (const FVertexID VertexID : SrcModelMeshDescription->Vertices().GetElementIDs())
+		for (const FVertexID VertexID : MeshDescription->Vertices().GetElementIDs())
 		{
 			VertexPositionsDup[VertexIndex] = VertexPositions[VertexID];
 		}
@@ -177,9 +174,9 @@ bool MeshPaintHelpers::PropagateColorsToRawMesh(UStaticMesh* StaticMesh, int32 L
 		);
 		if (NewVertexColors.Num() == NumVertex)
 		{
-			for (const FVertexInstanceID VertexInstanceID : SrcModelMeshDescription->VertexInstances().GetElementIDs())
+			for (const FVertexInstanceID VertexInstanceID : MeshDescription->VertexInstances().GetElementIDs())
 			{
-				const FVertexID VertexID = SrcModelMeshDescription->GetVertexInstanceVertex(VertexInstanceID);
+				const FVertexID VertexID = MeshDescription->GetVertexInstanceVertex(VertexInstanceID);
 				Colors[VertexInstanceID] = FLinearColor(NewVertexColors[VertexID.GetValue()]);
 			}
 			StaticMesh->CommitMeshDescription(LODIndex);
diff --git a/Engine/Source/Editor/UnrealEd/Private/Fbx/FbxMainExport.cpp b/Engine/Source/Editor/UnrealEd/Private/Fbx/FbxMainExport.cpp
--- a/Engine/Source/Editor/UnrealEd/Private/Fbx/FbxMainExport.cpp
+++ b/Engine/Source/Editor/UnrealEd/Private/Fbx/FbxMainExport.cpp
@@ -67,9 +67,7 @@
 #include "Matinee/MatineeActor.h"
 #include "FbxExporter.h"
 
-#include "MeshDescription.h"
-#include "MeshAttributes.h"
-#include "MeshAttributeArray.h"
+#include "StaticMeshAttributes.h"
 #include "MeshDescriptionOperations.h"
 
 #include "Components/BrushComponent.h"
@@ -695,7 +693,8 @@ void FFbxExporter::ExportBrush(ABrush* Actor, UModel* InModel, bool bConvertToSt
 	else
 	{
 		FMeshDescription Mesh;
-		UStaticMesh::RegisterMeshAttributes(Mesh);
+		FStaticMeshAttributes MeshAttributes(Mesh);
+		MeshAttributes.Register();
 		TArray<FStaticMaterial>	Materials;
 		GetBrushMesh(Actor,Actor->Brush,Mesh,Materials);
 
@@ -1019,7 +1018,9 @@ void FFbxExporter::ExportBSP( UModel* Model, bool bSelectedOnly )
 			{
 				uint32 NumWedges = ExportData->NumFaces * 3;
 
-				UStaticMesh::RegisterMeshAttributes(Mesh);
+				FStaticMeshAttributes InitMeshAttributes(Mesh);
+				InitMeshAttributes.Register();
+
 				ExportData->bInitialised = true;
 				Mesh.Empty();
 				Mesh.ReserveNewVertices(ExportData->NumVerts);
@@ -1028,8 +1029,7 @@ void FFbxExporter::ExportBSP( UModel* Model, bool bSelectedOnly )
 				Mesh.ReserveNewEdges(ExportData->NumFaces * 2);
 				Mesh.ReserveNewPolygonGroups(Materials.Num());
 				//We need to get the PolygonGroupImportedMaterial
-				FStaticMeshDescriptionAttributeGetter InitMeshAttributes(&Mesh);
-				TPolygonGroupAttributesRef<FName> InitPolygonGroupNames = InitMeshAttributes.GetPolygonGroupImportedMaterialSlotNames();
+				TPolygonGroupAttributesRef<FName> InitPolygonGroupNames = InitMeshAttributes.GetPolygonGroupMaterialSlotNames();
 				for (const FStaticMaterial& StaticMaterial : Materials)
 				{
 					const FPolygonGroupID PolygonGroupID = Mesh.CreatePolygonGroup();
@@ -1039,11 +1039,11 @@ void FFbxExporter::ExportBSP( UModel* Model, bool bSelectedOnly )
 			}
 			
 			//Get the Attributes
-			FStaticMeshDescriptionAttributeGetter MeshAttributes(&Mesh);
-			TVertexAttributesRef<FVector> VertexPositions = MeshAttributes.GetPositions();
-			TVertexInstanceAttributesRef<FVector2D> UVs = MeshAttributes.GetUVs();
-			TVertexInstanceAttributesRef<FVector4> Colors = MeshAttributes.GetColors();
-			TVertexInstanceAttributesRef<FVector> Normals = MeshAttributes.GetNormals();
+			FStaticMeshAttributes MeshAttributes(Mesh);
+			TVertexAttributesRef<FVector> VertexPositions = MeshAttributes.GetVertexPositions();
+			TVertexInstanceAttributesRef<FVector2D> UVs = MeshAttributes.GetVertexInstanceUVs();
+			TVertexInstanceAttributesRef<FVector4> Colors = MeshAttributes.GetVertexInstanceColors();
+			TVertexInstanceAttributesRef<FVector> Normals = MeshAttributes.GetVertexInstanceNormals();
 			TEdgeAttributesRef<bool> EdgeHardnesses = MeshAttributes.GetEdgeHardnesses();
 			TEdgeAttributesRef<float> EdgeCreaseSharpnesses = MeshAttributes.GetEdgeCreaseSharpnesses();
 
diff --git a/Engine/Source/Editor/UnrealEd/Private/Fbx/FbxStaticMeshImport.cpp b/Engine/Source/Editor/UnrealEd/Private/Fbx/FbxStaticMeshImport.cpp
--- a/Engine/Source/Editor/UnrealEd/Private/Fbx/FbxStaticMeshImport.cpp
+++ b/Engine/Source/Editor/UnrealEd/Private/Fbx/FbxStaticMeshImport.cpp
@@ -32,8 +32,7 @@
 #include "Interfaces/ITargetPlatformManagerModule.h"
 #include "Misc/FbxErrors.h"
 #include "PhysicsEngine/BodySetup.h"
-#include "MeshDescription.h"
-#include "MeshAttributes.h"
+#include "StaticMeshAttributes.h"
 #include "MeshDescriptionOperations.h"
 #include "IMeshBuilderModule.h"
 #include "Settings/EditorExperimentalSettings.h"
@@ -284,6 +283,8 @@ bool UnFbx::FFbxImporter::BuildStaticMeshFromGeometry(FbxNode* Node, UStaticMesh
 	FMeshDescription* MeshDescription = StaticMesh->GetMeshDescription(LODIndex);
 	//The mesh description should have been created before calling BuildStaticMeshFromGeometry
 	check(MeshDescription);
+	FStaticMeshAttributes Attributes(*MeshDescription);
+
 	//remove the bad polygons before getting any data from mesh
 	Mesh->RemoveBadPolygons();
 
@@ -553,15 +554,15 @@ bool UnFbx::FFbxImporter::BuildStaticMeshFromGeometry(FbxNode* Node, UStaticMesh
 	bool OddNegativeScale = IsOddNegativeScale(TotalMatrix);
 	bool bHasNonDegeneratePolygons = false;
 
-	TVertexAttributesRef<FVector> VertexPositions = MeshDescription->VertexAttributes().GetAttributesRef<FVector>(MeshAttribute::Vertex::Position);
-	TVertexInstanceAttributesRef<FVector> VertexInstanceNormals = MeshDescription->VertexInstanceAttributes().GetAttributesRef<FVector>(MeshAttribute::VertexInstance::Normal);
-	TVertexInstanceAttributesRef<FVector> VertexInstanceTangents = MeshDescription->VertexInstanceAttributes().GetAttributesRef<FVector>(MeshAttribute::VertexInstance::Tangent);
-	TVertexInstanceAttributesRef<float> VertexInstanceBinormalSigns = MeshDescription->VertexInstanceAttributes().GetAttributesRef<float>(MeshAttribute::VertexInstance::BinormalSign);
-	TVertexInstanceAttributesRef<FVector4> VertexInstanceColors = MeshDescription->VertexInstanceAttributes().GetAttributesRef<FVector4>(MeshAttribute::VertexInstance::Color);
-	TVertexInstanceAttributesRef<FVector2D> VertexInstanceUVs = MeshDescription->VertexInstanceAttributes().GetAttributesRef<FVector2D>(MeshAttribute::VertexInstance::TextureCoordinate);
-	TEdgeAttributesRef<bool> EdgeHardnesses = MeshDescription->EdgeAttributes().GetAttributesRef<bool>(MeshAttribute::Edge::IsHard);
-	TEdgeAttributesRef<float> EdgeCreaseSharpnesses = MeshDescription->EdgeAttributes().GetAttributesRef<float>(MeshAttribute::Edge::CreaseSharpness);
-	TPolygonGroupAttributesRef<FName> PolygonGroupImportedMaterialSlotNames = MeshDescription->PolygonGroupAttributes().GetAttributesRef<FName>(MeshAttribute::PolygonGroup::ImportedMaterialSlotName);
+	TVertexAttributesRef<FVector> VertexPositions = Attributes.GetVertexPositions();
+	TVertexInstanceAttributesRef<FVector> VertexInstanceNormals = Attributes.GetVertexInstanceNormals();
+	TVertexInstanceAttributesRef<FVector> VertexInstanceTangents = Attributes.GetVertexInstanceTangents();
+	TVertexInstanceAttributesRef<float> VertexInstanceBinormalSigns = Attributes.GetVertexInstanceBinormalSigns();
+	TVertexInstanceAttributesRef<FVector4> VertexInstanceColors = Attributes.GetVertexInstanceColors();
+	TVertexInstanceAttributesRef<FVector2D> VertexInstanceUVs = Attributes.GetVertexInstanceUVs();
+	TEdgeAttributesRef<bool> EdgeHardnesses = Attributes.GetEdgeHardnesses();
+	TEdgeAttributesRef<float> EdgeCreaseSharpnesses = Attributes.GetEdgeCreaseSharpnesses();
+	TPolygonGroupAttributesRef<FName> PolygonGroupImportedMaterialSlotNames = Attributes.GetPolygonGroupMaterialSlotNames();
 
 	int32 VertexOffset = MeshDescription->Vertices().Num();
 	int32 VertexInstanceOffset = MeshDescription->VertexInstances().Num();
@@ -1461,6 +1462,8 @@ UStaticMesh* UnFbx::FFbxImporter::ImportStaticMeshAsSingle(UObject* InParent, TA
 		MeshDescription->Empty();
 	}
 
+	FStaticMeshAttributes Attributes(*MeshDescription);
+
 	// make sure it has a new lighting guid
 	StaticMesh->LightingGuid = FGuid::NewGuid();
 
@@ -1508,9 +1511,9 @@ UStaticMesh* UnFbx::FFbxImporter::ImportStaticMeshAsSingle(UObject* InParent, TA
 
 	if (bBuildStatus)
 	{
-		TVertexInstanceAttributesRef<FVector2D> VertexInstanceUVs = MeshDescription->VertexInstanceAttributes().GetAttributesRef<FVector2D>(MeshAttribute::VertexInstance::TextureCoordinate);
+		TVertexInstanceAttributesRef<FVector2D> VertexInstanceUVs = Attributes.GetVertexInstanceUVs();
 		int32 FirstOpenUVChannel = VertexInstanceUVs.GetNumIndices() >= MAX_MESH_TEXTURE_COORDS_MD ? 1 : VertexInstanceUVs.GetNumIndices();
-		TPolygonGroupAttributesRef<FName> PolygonGroupImportedMaterialSlotNames = MeshDescription->PolygonGroupAttributes().GetAttributesRef<FName>(MeshAttribute::PolygonGroup::ImportedMaterialSlotName);
+		TPolygonGroupAttributesRef<FName> PolygonGroupImportedMaterialSlotNames = Attributes.GetPolygonGroupMaterialSlotNames();
 
 		TArray<FStaticMaterial> MaterialToAdd;
 		for (const FPolygonGroupID PolygonGroupID : MeshDescription->PolygonGroups().GetElementIDs())
@@ -1899,12 +1902,6 @@ void UnFbx::FFbxImporter::PostImportStaticMesh(UStaticMesh* StaticMesh, TArray<F
 	IConsoleVariable* CVarDistanceFieldInterface = IConsoleManager::Get().FindConsoleVariable(TEXT("r.GenerateMeshDistanceFields"));
 	bool bOriginalGenerateMeshDistanceField = StaticMesh->bGenerateMeshDistanceField;
 	
-	//Always triangulate the original mesh description after we import it
-	FMeshDescription* MeshDescription = StaticMesh->GetMeshDescription(LODIndex);
-	if (MeshDescription)
-	{
-		MeshDescription->TriangulateMesh();
-	}
 	bool bUserCancel = false;
 	if (StaticMesh->FixLODRequiresAdjacencyInformation(LODIndex))
 	{
diff --git a/Engine/Source/Runtime/MeshDescription/Public/MeshAttributes.h b/Engine/Source/Runtime/MeshDescription/Public/MeshAttributes.h
--- a/Engine/Source/Runtime/MeshDescription/Public/MeshAttributes.h
+++ b/Engine/Source/Runtime/MeshDescription/Public/MeshAttributes.h
@@ -3,46 +3,49 @@
 #pragma once
 
 #include "CoreMinimal.h"
+#include "MeshDescription.h"
 
-// @todo mesh description: for now, these attributes are hardcoded in the MeshDescription module.
-// It would be better if these were defined in separate modules, one per mesh type which uses MeshDescription.
-// This would give an extensible and modular structure whereby a mesh type defines its valid attributes and adapters.
 namespace MeshAttribute
 {
 	namespace Vertex
 	{
 		extern MESHDESCRIPTION_API const FName Position;
-		extern MESHDESCRIPTION_API const FName CornerSharpness;
 	}
+}
 
-	namespace VertexInstance
-	{
-		extern MESHDESCRIPTION_API const FName TextureCoordinate;
-		extern MESHDESCRIPTION_API const FName Normal;
-		extern MESHDESCRIPTION_API const FName Tangent;
-		extern MESHDESCRIPTION_API const FName BinormalSign;
-		extern MESHDESCRIPTION_API const FName Color;
-	}
 
-	namespace Edge
-	{
-		extern MESHDESCRIPTION_API const FName IsHard;
-		extern MESHDESCRIPTION_API const FName IsUVSeam;
-		extern MESHDESCRIPTION_API const FName CreaseSharpness;
-	}
+class FMeshAttributes
+{
+public:
+	explicit FMeshAttributes(FMeshDescription& InMeshDescription)
+		: MeshDescription(InMeshDescription)
+	{}
 
-	namespace Polygon
-	{
-		extern MESHDESCRIPTION_API const FName Normal;
-		extern MESHDESCRIPTION_API const FName Tangent;
-		extern MESHDESCRIPTION_API const FName Binormal;
-		extern MESHDESCRIPTION_API const FName Center;
-	}
+	virtual ~FMeshAttributes() = default;
 
-	namespace PolygonGroup
-	{
-		extern MESHDESCRIPTION_API const FName ImportedMaterialSlotName;
-		extern MESHDESCRIPTION_API const FName EnableCollision;
-		extern MESHDESCRIPTION_API const FName CastShadow;
-	}
-}
+	MESHDESCRIPTION_API virtual void Register();
+
+	/** Accessors for cached vertex position array */
+	TVertexAttributesRef<FVector> GetVertexPositions() { return MeshDescription.VertexAttributes().GetAttributesRef<FVector>(MeshAttribute::Vertex::Position); }
+	TVertexAttributesConstRef<FVector> GetVertexPositions() const { return MeshDescription.VertexAttributes().GetAttributesRef<FVector>(MeshAttribute::Vertex::Position); }
+
+protected:
+
+	FMeshDescription& MeshDescription;
+};
+
+
+class FMeshConstAttributes
+{
+public:
+	explicit FMeshConstAttributes(const FMeshDescription& InMeshDescription)
+		: MeshDescription(InMeshDescription)
+	{}
+
+	/** Accessors for cached vertex position array */
+	TVertexAttributesConstRef<FVector> GetVertexPositions() const { return MeshDescription.VertexAttributes().GetAttributesRef<FVector>(MeshAttribute::Vertex::Position); }
+
+protected:
+
+	const FMeshDescription& MeshDescription;
+};
diff --git a/Engine/Source/Runtime/StaticMeshDescription/Public/StaticMeshAttributes.h b/Engine/Source/Runtime/StaticMeshDescription/Public/StaticMeshAttributes.h
new file mode 100644
--- /dev/null
+++ b/Engine/Source/Runtime/StaticMeshDescription/Public/StaticMeshAttributes.h
@@ -0,0 +1,120 @@
+// Copyright 1998-2019 Epic Games, Inc. All Rights Reserved.
+
+#pragma once
+
+#include "MeshAttributes.h"
+
+
+namespace MeshAttribute
+{
+	namespace Vertex
+	{
+		extern STATICMESHDESCRIPTION_API const FName CornerSharpness;
+	}
+
+	namespace VertexInstance
+	{
+		extern STATICMESHDESCRIPTION_API const FName TextureCoordinate;
+		extern STATICMESHDESCRIPTION_API const FName Normal;
+		extern STATICMESHDESCRIPTION_API const FName Tangent;
+		extern STATICMESHDESCRIPTION_API const FName BinormalSign;
+		extern STATICMESHDESCRIPTION_API const FName Color;
+	}
+
+	namespace Edge
+	{
+		extern STATICMESHDESCRIPTION_API const FName IsHard;
+		extern STATICMESHDESCRIPTION_API const FName IsUVSeam;
+		extern STATICMESHDESCRIPTION_API const FName CreaseSharpness;
+	}
+
+	namespace Polygon
+	{
+		extern STATICMESHDESCRIPTION_API const FName Normal;
+		extern STATICMESHDESCRIPTION_API const FName Tangent;
+		extern STATICMESHDESCRIPTION_API const FName Binormal;
+		extern STATICMESHDESCRIPTION_API const FName Center;
+	}
+
+	namespace PolygonGroup
+	{
+		extern STATICMESHDESCRIPTION_API const FName ImportedMaterialSlotName;
+		extern STATICMESHDESCRIPTION_API const FName EnableCollision;
+		extern STATICMESHDESCRIPTION_API const FName CastShadow;
+	}
+}
+
+
+class FStaticMeshAttributes : public FMeshAttributes
+{
+public:
+
+	explicit FStaticMeshAttributes(FMeshDescription& InMeshDescription)
+		: FMeshAttributes(InMeshDescription)
+	{}
+
+	STATICMESHDESCRIPTION_API virtual void Register() override;
+
+	TVertexAttributesRef<float> GetVertexCornerSharpnesses() { return MeshDescription.VertexAttributes().GetAttributesRef<float>(MeshAttribute::Vertex::CornerSharpness); }
+	TVertexAttributesConstRef<float> GetVertexCornerSharpnesses() const { return MeshDescription.VertexAttributes().GetAttributesRef<float>(MeshAttribute::Vertex::CornerSharpness); }
+
+	TVertexInstanceAttributesRef<FVector2D> GetVertexInstanceUVs() { return MeshDescription.VertexInstanceAttributes().GetAttributesRef<FVector2D>(MeshAttribute::VertexInstance::TextureCoordinate); }
+	TVertexInstanceAttributesConstRef<FVector2D> GetVertexInstanceUVs() const { return MeshDescription.VertexInstanceAttributes().GetAttributesRef<FVector2D>(MeshAttribute::VertexInstance::TextureCoordinate); }
+
+	TVertexInstanceAttributesRef<FVector> GetVertexInstanceNormals() { return MeshDescription.VertexInstanceAttributes().GetAttributesRef<FVector>(MeshAttribute::VertexInstance::Normal); }
+	TVertexInstanceAttributesConstRef<FVector> GetVertexInstanceNormals() const { return MeshDescription.VertexInstanceAttributes().GetAttributesRef<FVector>(MeshAttribute::VertexInstance::Normal); }
+
+	TVertexInstanceAttributesRef<FVector> GetVertexInstanceTangents() { return MeshDescription.VertexInstanceAttributes().GetAttributesRef<FVector>(MeshAttribute::VertexInstance::Tangent); }
+	TVertexInstanceAttributesConstRef<FVector> GetVertexInstanceTangents() const { return MeshDescription.VertexInstanceAttributes().GetAttributesRef<FVector>(MeshAttribute::VertexInstance::Tangent); }
+
+	TVertexInstanceAttributesRef<float> GetVertexInstanceBinormalSigns() { return MeshDescription.VertexInstanceAttributes().GetAttributesRef<float>(MeshAttribute::VertexInstance::BinormalSign); }
+	TVertexInstanceAttributesConstRef<float> GetVertexInstanceBinormalSigns() const { return MeshDescription.VertexInstanceAttributes().GetAttributesRef<float>(MeshAttribute::VertexInstance::BinormalSign); }
+
+	TVertexInstanceAttributesRef<FVector4> GetVertexInstanceColors() { return MeshDescription.VertexInstanceAttributes().GetAttributesRef<FVector4>(MeshAttribute::VertexInstance::Color); }
+	TVertexInstanceAttributesConstRef<FVector4> GetVertexInstanceColors() const { return MeshDescription.VertexInstanceAttributes().GetAttributesRef<FVector4>(MeshAttribute::VertexInstance::Color); }
+
+	TEdgeAttributesRef<bool> GetEdgeHardnesses() { return MeshDescription.EdgeAttributes().GetAttributesRef<bool>(MeshAttribute::Edge::IsHard); }
+	TEdgeAttributesConstRef<bool> GetEdgeHardnesses() const { return MeshDescription.EdgeAttributes().GetAttributesRef<bool>(MeshAttribute::Edge::IsHard); }
+
+	TEdgeAttributesRef<float> GetEdgeCreaseSharpnesses() { return MeshDescription.EdgeAttributes().GetAttributesRef<float>(MeshAttribute::Edge::CreaseSharpness); }
+	TEdgeAttributesConstRef<float> GetEdgeCreaseSharpnesses() const { return MeshDescription.EdgeAttributes().GetAttributesRef<float>(MeshAttribute::Edge::CreaseSharpness); }
+
+	TPolygonAttributesRef<FVector> GetPolygonNormals() { return MeshDescription.PolygonAttributes().GetAttributesRef<FVector>(MeshAttribute::Polygon::Normal); }
+	TPolygonAttributesConstRef<FVector> GetPolygonNormals() const { return MeshDescription.PolygonAttributes().GetAttributesRef<FVector>(MeshAttribute::Polygon::Normal); }
+
+	TPolygonAttributesRef<FVector> GetPolygonTangents() { return MeshDescription.PolygonAttributes().GetAttributesRef<FVector>(MeshAttribute::Polygon::Tangent); }
+	TPolygonAttributesConstRef<FVector> GetPolygonTangents() const { return MeshDescription.PolygonAttributes().GetAttributesRef<FVector>(MeshAttribute::Polygon::Tangent); }
+
+	TPolygonAttributesRef<FVector> GetPolygonBinormals() { return MeshDescription.PolygonAttributes().GetAttributesRef<FVector>(MeshAttribute::Polygon::Binormal); }
+	TPolygonAttributesConstRef<FVector> GetPolygonBinormals() const { return MeshDescription.PolygonAttributes().GetAttributesRef<FVector>(MeshAttribute::Polygon::Binormal); }
+
+	TPolygonAttributesRef<FVector> GetPolygonCenters() { return MeshDescription.PolygonAttributes().GetAttributesRef<FVector>(MeshAttribute::Polygon::Center); }
+	TPolygonAttributesConstRef<FVector> GetPolygonCenters() const { return MeshDescription.PolygonAttributes().GetAttributesRef<FVector>(MeshAttribute::Polygon::Center); }
+
+	TPolygonGroupAttributesRef<FName> GetPolygonGroupMaterialSlotNames() { return MeshDescription.PolygonGroupAttributes().GetAttributesRef<FName>(MeshAttribute::PolygonGroup::ImportedMaterialSlotName); }
+	TPolygonGroupAttributesConstRef<FName> GetPolygonGroupMaterialSlotNames() const { return MeshDescription.PolygonGroupAttributes().GetAttributesRef<FName>(MeshAttribute::PolygonGroup::ImportedMaterialSlotName); }
+};
+
+
+class FStaticMeshConstAttributes : public FMeshConstAttributes
+{
+public:
+
+	explicit FStaticMeshConstAttributes(const FMeshDescription& InMeshDescription)
+		: FMeshConstAttributes(InMeshDescription)
+	{}
+
+	TVertexAttributesConstRef<float> GetVertexCornerSharpnesses() const { return MeshDescription.VertexAttributes().GetAttributesRef<float>(MeshAttribute::Vertex::CornerSharpness); }
+	TVertexInstanceAttributesConstRef<FVector2D> GetVertexInstanceUVs() const { return MeshDescription.VertexInstanceAttributes().GetAttributesRef<FVector2D>(MeshAttribute::VertexInstance::TextureCoordinate); }
+	TVertexInstanceAttributesConstRef<FVector> GetVertexInstanceNormals() const { return MeshDescription.VertexInstanceAttributes().GetAttributesRef<FVector>(MeshAttribute::VertexInstance::Normal); }
+	TVertexInstanceAttributesConstRef<FVector> GetVertexInstanceTangents() const { return MeshDescription.VertexInstanceAttributes().GetAttributesRef<FVector>(MeshAttribute::VertexInstance::Tangent); }
+	TVertexInstanceAttributesConstRef<float> GetVertexInstanceBinormalSigns() const { return MeshDescription.VertexInstanceAttributes().GetAttributesRef<float>(MeshAttribute::VertexInstance::BinormalSign); }
+	TVertexInstanceAttributesConstRef<FVector4> GetVertexInstanceColors() const { return MeshDescription.VertexInstanceAttributes().GetAttributesRef<FVector4>(MeshAttribute::VertexInstance::Color); }
+	TEdgeAttributesConstRef<bool> GetEdgeHardnesses() const { return MeshDescription.EdgeAttributes().GetAttributesRef<bool>(MeshAttribute::Edge::IsHard); }
+	TEdgeAttributesConstRef<float> GetEdgeCreaseSharpnesses() const { return MeshDescription.EdgeAttributes().GetAttributesRef<float>(MeshAttribute::Edge::CreaseSharpness); }
+	TPolygonAttributesConstRef<FVector> GetPolygonNormals() const { return MeshDescription.PolygonAttributes().GetAttributesRef<FVector>(MeshAttribute::Polygon::Normal); }
+	TPolygonAttributesConstRef<FVector> GetPolygonTangents() const { return MeshDescription.PolygonAttributes().GetAttributesRef<FVector>(MeshAttribute::Polygon::Tangent); }
+	TPolygonAttributesConstRef<FVector> GetPolygonBinormals() const { return MeshDescription.PolygonAttributes().GetAttributesRef<FVector>(MeshAttribute::Polygon::Binormal); }
+	TPolygonAttributesConstRef<FVector> GetPolygonCenters() const { return MeshDescription.PolygonAttributes().GetAttributesRef<FVector>(MeshAttribute::Polygon::Center); }
+	TPolygonGroupAttributesConstRef<FName> GetPolygonGroupMaterialSlotNames() const { return MeshDescription.PolygonGroupAttributes().GetAttributesRef<FName>(MeshAttribute::PolygonGroup::ImportedMaterialSlotName); }
+};
