Fix linux vtable compile error
#rnx

diff --git a/Engine/Source/Runtime/MeshDescription/Public/MeshAttributes.h b/Engine/Source/Runtime/MeshDescription/Public/MeshAttributes.h
--- a/Engine/Source/Runtime/MeshDescription/Public/MeshAttributes.h
+++ b/Engine/Source/Runtime/MeshDescription/Public/MeshAttributes.h
@@ -14,7 +14,7 @@ namespace MeshAttribute
 }
 
 
-class FMeshAttributes
+class MESHDESCRIPTION_API FMeshAttributes
 {
 public:
 	explicit FMeshAttributes(FMeshDescription& InMeshDescription)
@@ -23,7 +23,7 @@ public:
 
 	virtual ~FMeshAttributes() = default;
 
-	MESHDESCRIPTION_API virtual void Register();
+	virtual void Register();
 
 	/** Accessors for cached vertex position array */
 	TVertexAttributesRef<FVector> GetVertexPositions() { return MeshDescription.VertexAttributes().GetAttributesRef<FVector>(MeshAttribute::Vertex::Position); }
diff --git a/Engine/Source/Runtime/StaticMeshDescription/Public/StaticMeshAttributes.h b/Engine/Source/Runtime/StaticMeshDescription/Public/StaticMeshAttributes.h
--- a/Engine/Source/Runtime/StaticMeshDescription/Public/StaticMeshAttributes.h
+++ b/Engine/Source/Runtime/StaticMeshDescription/Public/StaticMeshAttributes.h
@@ -45,7 +45,7 @@ namespace MeshAttribute
 }
 
 
-class FStaticMeshAttributes : public FMeshAttributes
+class STATICMESHDESCRIPTION_API FStaticMeshAttributes : public FMeshAttributes
 {
 public:
 
@@ -53,7 +53,7 @@ public:
 		: FMeshAttributes(InMeshDescription)
 	{}
 
-	STATICMESHDESCRIPTION_API virtual void Register() override;
+	virtual void Register() override;
 
 	TVertexAttributesRef<float> GetVertexCornerSharpnesses() { return MeshDescription.VertexAttributes().GetAttributesRef<float>(MeshAttribute::Vertex::CornerSharpness); }
 	TVertexAttributesConstRef<float> GetVertexCornerSharpnesses() const { return MeshDescription.VertexAttributes().GetAttributesRef<float>(MeshAttribute::Vertex::CornerSharpness); }
