Fixes for various merge issues.

diff --git a/Engine/Plugins/Animation/LiveLink/Source/LiveLink/Public/LiveLinkClient.h b/Engine/Plugins/Animation/LiveLink/Source/LiveLink/Public/LiveLinkClient.h
--- a/Engine/Plugins/Animation/LiveLink/Source/LiveLink/Public/LiveLinkClient.h
+++ b/Engine/Plugins/Animation/LiveLink/Source/LiveLink/Public/LiveLinkClient.h
@@ -114,7 +114,6 @@ public:
 	virtual TArray<FLiveLinkTime> GetSubjectFrameTimes(FLiveLinkSubjectName SubjectName) const override;
 	virtual ULiveLinkSourceSettings* GetSourceSettings(const FGuid& SourceGuid) const override;
 	virtual UObject* GetSubjectSettings(const FLiveLinkSubjectKey& SubjectKey) const override;
-	virtual ULiveLinkSourceSettings* GetSourceSettings(const FGuid& SourceGuid) const override;
 
 
 	virtual bool EvaluateFrameFromSource_AnyThread(const FLiveLinkSubjectKey& SubjectKey, TSubclassOf<ULiveLinkRole> Role, FLiveLinkSubjectFrameData& OutFrame) override;
diff --git a/Engine/Plugins/MovieScene/MovieRenderPipeline/Source/MovieRenderPipelineCore/Private/MoviePipelineAudioRendering.cpp b/Engine/Plugins/MovieScene/MovieRenderPipeline/Source/MovieRenderPipelineCore/Private/MoviePipelineAudioRendering.cpp
--- a/Engine/Plugins/MovieScene/MovieRenderPipeline/Source/MovieRenderPipelineCore/Private/MoviePipelineAudioRendering.cpp
+++ b/Engine/Plugins/MovieScene/MovieRenderPipeline/Source/MovieRenderPipelineCore/Private/MoviePipelineAudioRendering.cpp
@@ -25,7 +25,7 @@ static FAudioDevice* GetAudioDeviceFromWorldContext(const UObject* WorldContextO
 		return nullptr;
 	}
 
-	return ThisWorld->GetAudioDevice();
+	return ThisWorld->GetAudioDevice().GetAudioDevice();
 }
 
 static Audio::FMixerDevice* GetAudioMixerDeviceFromWorldContext(const UObject* WorldContextObject)
@@ -152,7 +152,7 @@ void UMoviePipeline::ProcessAudioTick()
 		AudioDeltaTime = GetPipelineMasterConfig()->GetEffectiveFrameRate(TargetSequence).AsInterval();
 	}
 
-	if (FAudioDevice* AudioDevice = FAudioDeviceManager::Get()->GetActiveAudioDevice())
+	if (FAudioDevice* AudioDevice = FAudioDeviceManager::Get()->GetActiveAudioDevice().GetAudioDevice())
 	{
 		// Handle any game logic that changed Audio State.
 		AudioDevice->Update(true);
diff --git a/Engine/Source/Runtime/LiveLinkInterface/Public/ILiveLinkClient.h b/Engine/Source/Runtime/LiveLinkInterface/Public/ILiveLinkClient.h
--- a/Engine/Source/Runtime/LiveLinkInterface/Public/ILiveLinkClient.h
+++ b/Engine/Source/Runtime/LiveLinkInterface/Public/ILiveLinkClient.h
@@ -246,11 +246,6 @@ public:
 	 */
 	virtual TArray<FLiveLinkTime> GetSubjectFrameTimes(FLiveLinkSubjectName SubjectName) const = 0;
 
-	/**
-	 * Get the Settings of this source.
-	 */
-	virtual ULiveLinkSourceSettings* GetSourceSettings(const FGuid& SourceKey) const = 0;
-
 	/**
 	 * Get the Settings of this subject.
 	 * @note If subject is a VirtualSubject, the VirtualSubject itself is returned.
diff --git a/Engine/Source/Runtime/MovieScene/Private/MovieSceneTimeController.cpp b/Engine/Source/Runtime/MovieScene/Private/MovieSceneTimeController.cpp
--- a/Engine/Source/Runtime/MovieScene/Private/MovieSceneTimeController.cpp
+++ b/Engine/Source/Runtime/MovieScene/Private/MovieSceneTimeController.cpp
@@ -104,7 +104,7 @@ double FMovieSceneTimeController_PlatformClock::GetCurrentTime() const
 
 double FMovieSceneTimeController_AudioClock::GetCurrentTime() const
 {
-	FAudioDevice* AudioDevice = GEngine ? GEngine->GetMainAudioDevice() : nullptr;
+	FAudioDevice* AudioDevice = GEngine ? GEngine->GetMainAudioDevice().GetAudioDevice() : nullptr;
 	return AudioDevice ? AudioDevice->GetAudioClock() : FPlatformTime::Seconds();
 }
 
diff --git a/Engine/Source/Runtime/StaticMeshDescription/Private/StaticMeshOperations.cpp b/Engine/Source/Runtime/StaticMeshDescription/Private/StaticMeshOperations.cpp
--- a/Engine/Source/Runtime/StaticMeshDescription/Private/StaticMeshOperations.cpp
+++ b/Engine/Source/Runtime/StaticMeshDescription/Private/StaticMeshOperations.cpp
@@ -1100,13 +1100,7 @@ void FStaticMeshOperations::ComputeTangentsAndNormals(FMeshDescription& MeshDesc
 
 		//Set the corner angle data for a FVertexInstanceID
 		void SetCornerAngleData(FVertexInstanceID VertexInstanceID, float CornerAngle, int32 CornerIndex)
-			MeshDescription.DeletePolygonGroup(PolygonGroupID);
-		}
-		for (FVertexInstanceID VertexInstanceID : OrphanedVertexInstances)
 		{
-			MeshDescription.DeleteVertexInstance(VertexInstanceID, &OrphanedVertices);
-		}
-		for (FEdgeID EdgeID : OrphanedEdges)
 			CornerAngleDatas[CornerIndex].VertexInstanceID = VertexInstanceID;
 			CornerAngleDatas[CornerIndex].CornerAngle = CornerAngle;
 		}
@@ -1155,9 +1149,7 @@ void FStaticMeshOperations::ComputeTangentsAndNormals(FMeshDescription& MeshDesc
 		TriangleDatas.Reserve(MeshDescription.Triangles().Num());
 
 		for (FTriangleID TriangleID : MeshDescription.Triangles().GetElementIDs())
-			MeshDescription.DeleteEdge(EdgeID, &OrphanedVertices);
-		}
-		for (FVertexID VertexID : OrphanedVertices)
+		{
 			TArrayView<const FVertexInstanceID> VertexInstanceIDs = MeshDescription.GetTriangleVertexInstances(TriangleID);
 			//Triangle should use 3 vertex instances
 			check(VertexInstanceIDs.Num() == 3);
@@ -1187,8 +1179,7 @@ void FStaticMeshOperations::ComputeTangentsAndNormals(FMeshDescription& MeshDesc
 	//Iterate all vertex to compute normals for all vertex instance
 	ParallelFor(BatchCount,
 		[&Vertices, &BatchSize, &bComputeTangentWithMikkTSpace, &MeshDescription, bComputeWeightedNormals, &TriangleDatas](int32 BatchIndex)
-			MeshDescription.DeleteVertex(VertexID);
-		}
+		{
 			TVertexInstanceAttributesConstRef<FVector2D> VertexUVs = MeshDescription.VertexInstanceAttributes().GetAttributesRef<FVector2D>(MeshAttribute::VertexInstance::TextureCoordinate);
 			TVertexInstanceAttributesRef<FVector> VertexNormals = MeshDescription.VertexInstanceAttributes().GetAttributesRef<FVector>(MeshAttribute::VertexInstance::Normal);
 			TVertexInstanceAttributesRef<FVector> VertexTangents = MeshDescription.VertexInstanceAttributes().GetAttributesRef<FVector>(MeshAttribute::VertexInstance::Tangent);
