Integrating CL #11244454 from //UE4/Dev-Enterprise

Dataprep: Implement Filp Face Operation

diff --git a/Engine/Plugins/Enterprise/DataprepEditor/Source/DataprepLibraries/Private/DataprepOperations.cpp b/Engine/Plugins/Enterprise/DataprepEditor/Source/DataprepLibraries/Private/DataprepOperations.cpp
--- a/Engine/Plugins/Enterprise/DataprepEditor/Source/DataprepLibraries/Private/DataprepOperations.cpp
+++ b/Engine/Plugins/Enterprise/DataprepEditor/Source/DataprepLibraries/Private/DataprepOperations.cpp
@@ -349,4 +349,38 @@ void UDataprepRandomizeTransformOperation::OnExecution_Implementation(const FDat
 	UDataprepOperationsLibrary::RandomizeTransform(InContext.Objects, TransformType, ReferenceFrame, Min, Max);
 }
 
+void UDataprepFlipFacesOperation::OnExecution_Implementation(const FDataprepContext& InContext)
+{
+#ifdef LOG_TIME
+	DataprepOperationTime::FTimeLogger TimeLogger(TEXT("FlipFaces"), [&](FText Text) { this->LogInfo(Text); });
+#endif
+
+	TSet<UStaticMesh*> StaticMeshes;
+
+	// Re-create static meshes
+	for (UObject* Object : InContext.Objects)
+	{
+		if (AActor* Actor = Cast< AActor >(Object))
+		{
+			TInlineComponentArray<UStaticMeshComponent*> StaticMeshComponents(Actor);
+			for (UStaticMeshComponent* StaticMeshComponent : StaticMeshComponents)
+			{
+				UStaticMesh* StaticMesh = StaticMeshComponent->GetStaticMesh();
+
+				if (nullptr == StaticMesh)
+				{
+					continue;
+				}
+				StaticMeshes.Add(StaticMesh);
+			}
+		}
+	}
+
+	// Execute operation
+	UDataprepOperationsLibrary::FlipFaces(StaticMeshes);
+
+	// Re-create meshes render data
+	UStaticMesh::BatchBuild(StaticMeshes.Array());
+}
+
 #undef LOCTEXT_NAMESPACE
diff --git a/Engine/Plugins/Enterprise/DataprepEditor/Source/DataprepLibraries/Private/DataprepOperations.h b/Engine/Plugins/Enterprise/DataprepEditor/Source/DataprepLibraries/Private/DataprepOperations.h
--- a/Engine/Plugins/Enterprise/DataprepEditor/Source/DataprepLibraries/Private/DataprepOperations.h
+++ b/Engine/Plugins/Enterprise/DataprepEditor/Source/DataprepLibraries/Private/DataprepOperations.h
@@ -429,3 +429,22 @@ protected:
 	virtual void OnExecution_Implementation(const FDataprepContext& InContext) override;
 	//~ End UDataprepOperation Interface
 };
+
+UCLASS(Experimental, Category = ActorOperation, Meta = (DisplayName = "Flip Faces", ToolTip = "On each actor to process, flip faces of each mesh"))
+class UDataprepFlipFacesOperation : public UDataprepOperation
+{
+	GENERATED_BODY()
+
+public:
+
+	//~ Begin UDataprepOperation Interface
+public:
+	virtual FText GetCategory_Implementation() const override
+	{
+		return FDataprepOperationCategories::MeshOperation;
+	}
+
+protected:
+	virtual void OnExecution_Implementation(const FDataprepContext& InContext) override;
+	//~ End UDataprepOperation Interface
+};
diff --git a/Engine/Plugins/Enterprise/DataprepEditor/Source/DataprepLibraries/Private/DataprepOperationsLibrary.cpp b/Engine/Plugins/Enterprise/DataprepEditor/Source/DataprepLibraries/Private/DataprepOperationsLibrary.cpp
--- a/Engine/Plugins/Enterprise/DataprepEditor/Source/DataprepLibraries/Private/DataprepOperationsLibrary.cpp
+++ b/Engine/Plugins/Enterprise/DataprepEditor/Source/DataprepLibraries/Private/DataprepOperationsLibrary.cpp
@@ -27,6 +27,7 @@
 #include "Misc/FileHelper.h"
 #include "ObjectTools.h"
 #include "StaticMeshAttributes.h"
+#include "StaticMeshOperations.h"
 #include "TessellationRendering.h"
 #include "UObject/SoftObjectPath.h"
 
@@ -851,4 +852,24 @@ void UDataprepOperationsLibrary::RandomizeTransform(const TArray<UObject*>& Sele
 	}
 }
 
+void UDataprepOperationsLibrary::FlipFaces(const TSet< UStaticMesh* >& StaticMeshes)
+{
+	for (UStaticMesh* StaticMesh : StaticMeshes)
+	{
+		if (nullptr == StaticMesh || !StaticMesh->IsMeshDescriptionValid(0))
+		{
+			continue;
+		}
+
+		FMeshDescription* MeshDescription = StaticMesh->GetMeshDescription(0);
+
+		UStaticMesh::FCommitMeshDescriptionParams Params;
+		Params.bMarkPackageDirty = false;
+		Params.bUseHashAsGuid = true;
+
+		FStaticMeshOperations::FlipPolygons(*MeshDescription);
+		StaticMesh->CommitMeshDescription(0, Params);
+	}
+}
+
 #undef LOCTEXT_NAMESPACE
diff --git a/Engine/Plugins/Enterprise/DataprepEditor/Source/DataprepLibraries/Private/DataprepOperationsLibrary.h b/Engine/Plugins/Enterprise/DataprepEditor/Source/DataprepLibraries/Private/DataprepOperationsLibrary.h
--- a/Engine/Plugins/Enterprise/DataprepEditor/Source/DataprepLibraries/Private/DataprepOperationsLibrary.h
+++ b/Engine/Plugins/Enterprise/DataprepEditor/Source/DataprepLibraries/Private/DataprepOperationsLibrary.h
@@ -312,6 +312,13 @@ public:
 	UFUNCTION(BlueprintCallable, Category = "Dataprep | Operation")
 	static void RandomizeTransform(const TArray<UObject*>& SelectedObjects, ERandomizeTransformType TransformType, ERandomizeTransformReferenceFrame ReferenceFrame, const FVector& Min, const FVector& Max);
 
+	/**
+	 * Flip the faces of all elements of a set of Static Meshes or Static Mesh Actors
+	 * @param SelectedObjects	Objects to the flip the faces of
+	 */
+	UFUNCTION(BlueprintCallable, Category = "Dataprep | Operation")
+	static void FlipFaces( const TSet< UStaticMesh* >& StaticMeshes );
+
 private:
 	static void SubstituteMaterial(const TArray<UObject*>& SelectedObjects, const FString& MaterialSearch, EEditorScriptingStringMatchType StringMatch, const TArray<UMaterialInterface*>& MaterialList, UMaterialInterface* MaterialSubstitute);
 	static void SubstituteMesh(const TArray<UObject*>& SelectedObjects, const FString& MeshSearch, EEditorScriptingStringMatchType StringMatch, const TArray<UStaticMesh*>& MeshList, UStaticMesh* MeshSubstitute);
diff --git a/Engine/Source/Runtime/StaticMeshDescription/Private/StaticMeshOperations.cpp b/Engine/Source/Runtime/StaticMeshDescription/Private/StaticMeshOperations.cpp
--- a/Engine/Source/Runtime/StaticMeshDescription/Private/StaticMeshOperations.cpp
+++ b/Engine/Source/Runtime/StaticMeshDescription/Private/StaticMeshOperations.cpp
@@ -2235,4 +2235,29 @@ FSHAHash FStaticMeshOperations::ComputeSHAHash(const FMeshDescription& MeshDescr
 	return OutHash;
 }
 
+void FStaticMeshOperations::FlipPolygons(FMeshDescription& MeshDescription)
+{
+	TSet<FVertexInstanceID> VertexInstanceIDs;
+	for (FPolygonID PolygonID : MeshDescription.Polygons().GetElementIDs())
+	{
+		VertexInstanceIDs.Append(MeshDescription.GetPolygonVertexInstances(PolygonID));
+		MeshDescription.ReversePolygonFacing(PolygonID);
+	}
+
+	// Flip tangents and normals
+	const TVertexInstanceAttributesRef<FVector> VertexNormals = MeshDescription.VertexInstanceAttributes().GetAttributesRef<FVector>(MeshAttribute::VertexInstance::Normal);
+	const TVertexInstanceAttributesRef<FVector> VertexTangents = MeshDescription.VertexInstanceAttributes().GetAttributesRef<FVector>(MeshAttribute::VertexInstance::Tangent);
+
+	for (const FVertexInstanceID VertexInstanceID : VertexInstanceIDs)
+	{
+		// Just reverse the sign of the normals/tangents; note that since binormals are the cross product of normal with tangent, they are left untouched
+		FVector Normal = VertexNormals[VertexInstanceID] * -1.0f;
+		FVector Tangent = VertexTangents[VertexInstanceID] * -1.0f;
+
+		TAttributesSet<FVertexInstanceID>& AttributesSet = MeshDescription.VertexInstanceAttributes();
+		AttributesSet.SetAttribute(VertexInstanceID, MeshAttribute::VertexInstance::Normal, 0, Normal);
+		AttributesSet.SetAttribute(VertexInstanceID, MeshAttribute::VertexInstance::Tangent, 0, Tangent);
+	}
+}
+
 #undef LOCTEXT_NAMESPACE
\ No newline at end of file
diff --git a/Engine/Source/Runtime/StaticMeshDescription/Public/StaticMeshOperations.h b/Engine/Source/Runtime/StaticMeshDescription/Public/StaticMeshOperations.h
--- a/Engine/Source/Runtime/StaticMeshDescription/Public/StaticMeshOperations.h
+++ b/Engine/Source/Runtime/StaticMeshDescription/Public/StaticMeshOperations.h
@@ -148,4 +148,7 @@ public:
 
 	/** Computes the SHA hash of all the attributes values in the MeshDescription. */
 	static FSHAHash ComputeSHAHash(const FMeshDescription& MeshDescription);
+
+	/** Flip the facing for a set of input polygons. */
+	static void FlipPolygons(FMeshDescription& MeshDescription);
 };
