First pass of MeshDescription API and format refactor.

diff --git a/Engine/Source/Runtime/StaticMeshDescription/Private/StaticMeshOperations.cpp b/Engine/Source/Runtime/StaticMeshDescription/Private/StaticMeshOperations.cpp
--- a/Engine/Source/Runtime/StaticMeshDescription/Private/StaticMeshOperations.cpp
+++ b/Engine/Source/Runtime/StaticMeshDescription/Private/StaticMeshOperations.cpp
@@ -39,7 +39,7 @@ static bool GetPolygonTangentsAndNormals(FMeshDescription& MeshDescription,
 	FVector Center(FVector::ZeroVector);
 
 	// Calculate the center of this polygon
-	const TArray<FVertexInstanceID>& VertexInstanceIDs = MeshDescription.GetPolygonVertexInstances(PolygonID);
+	TArray<FVertexInstanceID, TInlineAllocator<4>> VertexInstanceIDs = MeshDescription.GetPolygonVertexInstances<TInlineAllocator<4>>(PolygonID);
 	for (const FVertexInstanceID VertexInstanceID : VertexInstanceIDs)
 	{
 		Center += VertexPositions[MeshDescription.GetVertexInstanceVertex(VertexInstanceID)];
@@ -126,8 +126,10 @@ static bool GetPolygonTangentsAndNormals(FMeshDescription& MeshDescription,
 	return bValidNTBs;
 }
 
+
 void FStaticMeshOperations::ComputePolygonTangentsAndNormals(FMeshDescription& MeshDescription, float ComparisonThreshold)
 {
+	PRAGMA_DISABLE_DEPRECATION_WARNINGS
 	TRACE_CPUPROFILER_EVENT_SCOPE(TEXT("FStaticMeshOperations::ComputePolygonTangentsAndNormals_Selection"));
 
 	FStaticMeshAttributes Attributes(MeshDescription);
@@ -174,8 +176,132 @@ void FStaticMeshOperations::ComputePolygonTangentsAndNormals(FMeshDescription& M
 			}
 		}
 	);
+	PRAGMA_ENABLE_DEPRECATION_WARNINGS
+}
+
+
+static TTuple<FVector, FVector, FVector> GetTriangleTangentsAndNormals(float ComparisonThreshold, TArrayView<const FVector> VertexPositions, TArrayView<const FVector2D> VertexUVs)
+{
+	float AdjustedComparisonThreshold = FMath::Max(ComparisonThreshold, MIN_flt);
+
+	const FVector Position0 = VertexPositions[0];
+	const FVector DPosition1 = VertexPositions[1] - Position0;
+	const FVector DPosition2 = VertexPositions[2] - Position0;
+
+	const FVector2D UV0 = VertexUVs[0];
+	const FVector2D DUV1 = VertexUVs[1] - UV0;
+	const FVector2D DUV2 = VertexUVs[2] - UV0;
+
+	// We have a left-handed coordinate system, but a counter-clockwise winding order
+	// Hence normal calculation has to take the triangle vectors cross product in reverse.
+	FVector Normal = FVector::CrossProduct(DPosition2, DPosition1).GetSafeNormal(AdjustedComparisonThreshold);
+	if (!Normal.IsNearlyZero(ComparisonThreshold))
+	{
+		FMatrix	ParameterToLocal(
+			DPosition1,
+			DPosition2,
+			Position0,
+			FVector::ZeroVector
+		);
+
+		FMatrix ParameterToTexture(
+			FPlane(DUV1.X, DUV1.Y, 0, 0),
+			FPlane(DUV2.X, DUV2.Y, 0, 0),
+			FPlane(UV0.X, UV0.Y, 1, 0),
+			FPlane(0, 0, 0, 1)
+		);
+
+		// Use InverseSlow to catch singular matrices.  Inverse can miss this sometimes.
+		const FMatrix TextureToLocal = ParameterToTexture.Inverse() * ParameterToLocal;
+
+		FVector Tangent = TextureToLocal.TransformVector(FVector(1, 0, 0)).GetSafeNormal();
+		FVector Binormal = TextureToLocal.TransformVector(FVector(0, 1, 0)).GetSafeNormal();
+		FVector::CreateOrthonormalBasis(Tangent, Binormal, Normal);
+
+		if (Tangent.IsNearlyZero() || Tangent.ContainsNaN()
+			|| Binormal.IsNearlyZero() || Binormal.ContainsNaN())
+		{
+			Tangent = FVector::ZeroVector;
+			Binormal = FVector::ZeroVector;
+		}
+
+		if (Normal.IsNearlyZero() || Normal.ContainsNaN())
+		{
+			Normal = FVector::ZeroVector;
+		}
+
+		return MakeTuple(Normal.GetSafeNormal(), Tangent.GetSafeNormal(), Binormal.GetSafeNormal());
+	}
+	else
+	{
+		// This will force a recompute of the normals and tangents
+		return MakeTuple(FVector::ZeroVector, FVector::ZeroVector, FVector::ZeroVector);
+	}
 }
 
+
+void FStaticMeshOperations::ComputeTriangleTangentsAndNormals(FMeshDescription& MeshDescription, float ComparisonThreshold)
+{
+	TRACE_CPUPROFILER_EVENT_SCOPE(TEXT("FStaticMeshOperations::ComputeTriangleTangentsAndNormals_Selection"));
+
+	FStaticMeshAttributes Attributes(MeshDescription);
+	Attributes.RegisterTriangleNormalAndTangentAttributes();
+
+	// Check that the mesh description is compact
+	const int32 NumTriangles = MeshDescription.Triangles().Num();
+	check(MeshDescription.Triangles().GetArraySize() == NumTriangles);
+
+	// Split work in batch to reduce call overhead
+	const int32 BatchSize = 8 * 1024;
+	const int32 BatchCount = (NumTriangles + BatchSize - 1) / BatchSize;
+
+	ParallelFor(BatchCount,
+		[BatchSize, ComparisonThreshold, NumTriangles, &Attributes](int32 BatchIndex)
+		{
+			TArrayView<const FVector> VertexPositions = Attributes.GetVertexPositions().GetRawArray();
+			TArrayView<const FVector2D> VertexUVs = Attributes.GetVertexInstanceUVs().GetRawArray();
+			TArrayView<const FVertexID> TriangleVertexIDs = Attributes.GetTriangleVertexIndices().GetRawArray();
+			TArrayView<const FVertexInstanceID> TriangleVertexInstanceIDs = Attributes.GetTriangleVertexInstanceIndices().GetRawArray();
+
+			TArrayView<FVector> TriangleNormals = Attributes.GetTriangleNormals().GetRawArray();
+			TArrayView<FVector> TriangleTangents = Attributes.GetTriangleTangents().GetRawArray();
+			TArrayView<FVector> TriangleBinormals = Attributes.GetTriangleBinormals().GetRawArray();
+
+			int32 StartIndex = BatchIndex * BatchSize;
+			int32 TriIndex = StartIndex * 3;
+			int32 EndIndex = FMath::Min(StartIndex + BatchSize, NumTriangles);
+			for (; StartIndex < EndIndex; ++StartIndex, TriIndex += 3)
+			{
+				if (!TriangleNormals[StartIndex].IsNearlyZero())
+				{
+					// Bypass normal calculation if it's already done
+					continue;
+				}
+
+				FVector TriangleVertexPositions[3] =
+				{
+					VertexPositions[TriangleVertexIDs[TriIndex]],
+					VertexPositions[TriangleVertexIDs[TriIndex + 1]],
+					VertexPositions[TriangleVertexIDs[TriIndex + 2]]
+				};
+
+				FVector2D TriangleUVs[3] =
+				{
+					VertexUVs[TriangleVertexInstanceIDs[TriIndex]],
+					VertexUVs[TriangleVertexInstanceIDs[TriIndex + 1]],
+					VertexUVs[TriangleVertexInstanceIDs[TriIndex + 2]]
+				};
+
+				TTuple<FVector, FVector, FVector> Result = GetTriangleTangentsAndNormals(ComparisonThreshold, TriangleVertexPositions, TriangleUVs);
+				TriangleNormals[StartIndex] = Result.Get<0>();
+				TriangleTangents[StartIndex] = Result.Get<1>();
+				TriangleBinormals[StartIndex] = Result.Get<2>();
+			}
+		}
+	);
+}
+
+
 void FStaticMeshOperations::DetermineEdgeHardnessesFromVertexInstanceNormals(FMeshDescription& MeshDescription, float Tolerance)
 {
 	FStaticMeshAttributes Attributes(MeshDescription);
@@ -247,8 +373,8 @@ struct FVertexInfo
 {
 	FVertexInfo()
 	{
-		TriangleID = FTriangleID::Invalid;
-		VertexInstanceID = FVertexInstanceID::Invalid;
+		TriangleID = INDEX_NONE;
+		VertexInstanceID = INDEX_NONE;
 		UVs = FVector2D(0.0f, 0.0f);
 	}
 
@@ -316,7 +442,7 @@ void FStaticMeshOperations::ConvertHardEdgesToSmoothGroup(const FMeshDescription
 		ConnectedPolygons.Reset();
 		LastConnectedPolygons.Reset();
 		ConnectedPolygons.Add(PolygonID);
-		LastConnectedPolygons.Add(FPolygonID::Invalid);
+		LastConnectedPolygons.Add(INDEX_NONE);
 		while (ConnectedPolygons.Num() > 0)
 		{
 			check(LastConnectedPolygons.Num() == ConnectedPolygons.Num());
@@ -330,7 +456,7 @@ void FStaticMeshOperations::ConvertHardEdgesToSmoothGroup(const FMeshDescription
 			uint32& SmoothGroup = PolygonSmoothGroup.FindOrAdd(CurrentPolygonID);
 			uint32 AvoidSmoothGroup = 0;
 			uint32 NeighborSmoothGroup = 0;
-			const uint32 LastSmoothGroupValue = (LastPolygonID == FPolygonID::Invalid) ? 0 : PolygonSmoothGroup[LastPolygonID];
+			const uint32 LastSmoothGroupValue = (LastPolygonID == INDEX_NONE) ? 0 : PolygonSmoothGroup[LastPolygonID];
 			PolygonEdges.Reset();
 			SourceMeshDescription.GetPolygonPerimeterEdges(CurrentPolygonID, PolygonEdges);
 			for (const FEdgeID& EdgeID : PolygonEdges)
@@ -518,7 +644,7 @@ void FStaticMeshOperations::ConvertToRawMesh(const FMeshDescription& SourceMeshD
 	DestinationRawMesh.WedgeTangentX.AddZeroed(WedgeIndexNumber);
 	DestinationRawMesh.WedgeTangentY.AddZeroed(WedgeIndexNumber);
 	DestinationRawMesh.WedgeTangentZ.AddZeroed(WedgeIndexNumber);
-	int32 ExistingUVCount = VertexInstanceUVs.GetNumIndices();
+	int32 ExistingUVCount = VertexInstanceUVs.GetNumChannels();
 	for (int32 UVIndex = 0; UVIndex < ExistingUVCount; ++UVIndex)
 	{
 		DestinationRawMesh.WedgeTexCoords[UVIndex].AddZeroed(WedgeIndexNumber);
@@ -530,7 +656,7 @@ void FStaticMeshOperations::ConvertToRawMesh(const FMeshDescription& SourceMeshD
 	{
 		const FPolygonGroupID& PolygonGroupID = SourceMeshDescription.GetPolygonPolygonGroup(PolygonID);
 		int32 PolygonIDValue = PolygonID.GetValue();
-		const TArray<FTriangleID>& TriangleIDs = SourceMeshDescription.GetPolygonTriangleIDs(PolygonID);
+		TArrayView<const FTriangleID> TriangleIDs = SourceMeshDescription.GetPolygonTriangleIDs(PolygonID);
 		for (const FTriangleID TriangleID : TriangleIDs)
 		{
 			if (MaterialMap.Num() > 0 && MaterialMap.Contains(PolygonGroupMaterialSlotName[PolygonGroupID]))
@@ -626,12 +752,12 @@ void FillMeshDescriptionVertexPositionNoDuplicate(const TArray<FVector>& RawMesh
 	UniqueVertexDone.AddUninitialized(VertexCount);
 	for (int32 VertexIndex = 0; VertexIndex < VertexCount; ++VertexIndex)
 	{
-		UniqueVertexDone[VertexIndex] = FVertexID::Invalid;
+		UniqueVertexDone[VertexIndex] = INDEX_NONE;
 	}
 	for (int32 VertexIndex = 0; VertexIndex < NumVertex; ++VertexIndex)
 	{
 		int32 RealIndex = TempRemapVertexPosition[VertexIndex];
-		if (UniqueVertexDone[RealIndex] != FVertexID::Invalid)
+		if (UniqueVertexDone[RealIndex] != INDEX_NONE)
 		{
 			RemapVertexPosition[VertexIndex] = UniqueVertexDone[RealIndex];
 			continue;
@@ -686,7 +812,7 @@ void FStaticMeshOperations::ConvertFromRawMesh(const FRawMesh& SourceRawMesh, FM
 			NumTexCoords++;
 		}
 	}
-	VertexInstanceUVs.SetNumIndices(NumTexCoords);
+	VertexInstanceUVs.SetNumChannels(NumTexCoords);
 
 	//Ensure we do not have any duplicate, We found all duplicated vertex and compact them and build a remap indice array to remap the wedgeindices
 	TArray<FVertexID> RemapVertexPosition;
@@ -731,7 +857,7 @@ void FStaticMeshOperations::ConvertFromRawMesh(const FRawMesh& SourceRawMesh, FM
 		}
 
 		//PolygonGroup
-		FPolygonGroupID PolygonGroupID = FPolygonGroupID::Invalid;
+		FPolygonGroupID PolygonGroupID = INDEX_NONE;
 		FName PolygonGroupImportedMaterialSlotName = NAME_None;
 		int32 MaterialIndex = SourceRawMesh.FaceMaterialIndices[TriangleIndex];
 		if (MaterialIndexToPolygonGroup.Contains(MaterialIndex))
@@ -751,15 +877,14 @@ void FStaticMeshOperations::ConvertFromRawMesh(const FRawMesh& SourceRawMesh, FM
 			}
 		}
 
-		if (PolygonGroupID == FPolygonGroupID::Invalid)
+		if (PolygonGroupID == INDEX_NONE)
 		{
 			PolygonGroupID = DestinationMeshDescription.CreatePolygonGroup();
 			PolygonGroupImportedMaterialSlotNames[PolygonGroupID] = PolygonGroupImportedMaterialSlotName == NAME_None ? FName(*FString::Printf(TEXT("MaterialSlot_%d"), MaterialIndex)) : PolygonGroupImportedMaterialSlotName;
 			PolygonGroups.Add(PolygonGroupID);
 			MaterialIndexToPolygonGroup.Add(MaterialIndex, PolygonGroupID);
 		}
-		TArray<FVertexInstanceID> TriangleVertexInstanceIDs;
-		TriangleVertexInstanceIDs.SetNum(3);
+		FVertexInstanceID TriangleVertexInstanceIDs[3];
 		for (int32 Corner = 0; Corner < 3; ++Corner)
 		{
 			int32 VerticeIndex = VerticeIndexBase + Corner;
@@ -794,7 +919,7 @@ void FStaticMeshOperations::ConvertFromRawMesh(const FRawMesh& SourceRawMesh, FM
 			}
 		}
 
-		DestinationMeshDescription.CreatePolygon(PolygonGroupID, TriangleVertexInstanceIDs);
+		DestinationMeshDescription.CreateTriangle(PolygonGroupID, TriangleVertexInstanceIDs);
 	}
 
 	ConvertSmoothGroupToHardEdges(SourceRawMesh.FaceSmoothingMasks, DestinationMeshDescription);
@@ -803,49 +928,39 @@ void FStaticMeshOperations::ConvertFromRawMesh(const FRawMesh& SourceRawMesh, FM
 	if (!bSkipNormalsAndTangents && (!bHasNormals || !bHasTangents))
 	{
 		//DestinationMeshDescription.ComputePolygonTangentsAndNormals(0.0f);
-		ComputePolygonTangentsAndNormals(DestinationMeshDescription, 0.0f);
+		ComputeTriangleTangentsAndNormals(DestinationMeshDescription, 0.0f);
 
 		//Create the missing normals and recompute the tangents with MikkTSpace.
 		EComputeNTBsFlags ComputeNTBsOptions = EComputeNTBsFlags::Tangents | EComputeNTBsFlags::UseMikkTSpace | EComputeNTBsFlags::BlendOverlappingNormals;
 		ComputeTangentsAndNormals(DestinationMeshDescription, ComputeNTBsOptions);
 	}
+
+	DestinationMeshDescription.BuildIndexers();
 }
 
 void FStaticMeshOperations::AppendMeshDescription(const FMeshDescription& SourceMesh, FMeshDescription& TargetMesh, FAppendSettings& AppendSettings)
 {
 	TRACE_CPUPROFILER_EVENT_SCOPE(FStaticMeshOperations::AppendMeshDescription);
 
-	//Vertex Attributes
-	TVertexAttributesConstRef<FVector> SourceVertexPositions = SourceMesh.VertexAttributes().GetAttributesRef<FVector>(MeshAttribute::Vertex::Position);
-	TVertexAttributesConstRef<float> SourceVertexCornerSharpness = SourceMesh.VertexAttributes().GetAttributesRef<float>(MeshAttribute::Vertex::CornerSharpness);
-
-	TVertexAttributesRef<FVector> TargetVertexPositions = TargetMesh.VertexAttributes().GetAttributesRef<FVector>(MeshAttribute::Vertex::Position);
-	TVertexAttributesRef<float> TargetVertexCornerSharpness = TargetMesh.VertexAttributes().GetAttributesRef<float>(MeshAttribute::Vertex::CornerSharpness);
-
-	//Edge Attributes
-	TEdgeAttributesConstRef<bool> SourceEdgeHardnesses = SourceMesh.EdgeAttributes().GetAttributesRef<bool>(MeshAttribute::Edge::IsHard);
-	TEdgeAttributesConstRef<float> SourceEdgeCreaseSharpnesses = SourceMesh.EdgeAttributes().GetAttributesRef<float>(MeshAttribute::Edge::CreaseSharpness);
-
-	TEdgeAttributesRef<bool> TargetEdgeHardnesses = TargetMesh.EdgeAttributes().GetAttributesRef<bool>(MeshAttribute::Edge::IsHard);
-	TEdgeAttributesRef<float> TargetEdgeCreaseSharpnesses = TargetMesh.EdgeAttributes().GetAttributesRef<float>(MeshAttribute::Edge::CreaseSharpness);
-
-	//PolygonGroup Attributes
-	TPolygonGroupAttributesConstRef<FName> SourceImportedMaterialSlotNames = SourceMesh.PolygonGroupAttributes().GetAttributesRef<FName>(MeshAttribute::PolygonGroup::ImportedMaterialSlotName);
-
-	TPolygonGroupAttributesRef<FName> TargetImportedMaterialSlotNames = TargetMesh.PolygonGroupAttributes().GetAttributesRef<FName>(MeshAttribute::PolygonGroup::ImportedMaterialSlotName);
-
-	//VertexInstance Attributes
-	TVertexInstanceAttributesConstRef<FVector> SourceVertexInstanceNormals = SourceMesh.VertexInstanceAttributes().GetAttributesRef<FVector>(MeshAttribute::VertexInstance::Normal);
-	TVertexInstanceAttributesConstRef<FVector> SourceVertexInstanceTangents = SourceMesh.VertexInstanceAttributes().GetAttributesRef<FVector>(MeshAttribute::VertexInstance::Tangent);
-	TVertexInstanceAttributesConstRef<float> SourceVertexInstanceBinormalSigns = SourceMesh.VertexInstanceAttributes().GetAttributesRef<float>(MeshAttribute::VertexInstance::BinormalSign);
-	TVertexInstanceAttributesConstRef<FVector4> SourceVertexInstanceColors = SourceMesh.VertexInstanceAttributes().GetAttributesRef<FVector4>(MeshAttribute::VertexInstance::Color);
-	TVertexInstanceAttributesConstRef<FVector2D> SourceVertexInstanceUVs = SourceMesh.VertexInstanceAttributes().GetAttributesRef<FVector2D>(MeshAttribute::VertexInstance::TextureCoordinate);
-
-	TVertexInstanceAttributesRef<FVector> TargetVertexInstanceNormals = TargetMesh.VertexInstanceAttributes().GetAttributesRef<FVector>(MeshAttribute::VertexInstance::Normal);
-	TVertexInstanceAttributesRef<FVector> TargetVertexInstanceTangents = TargetMesh.VertexInstanceAttributes().GetAttributesRef<FVector>(MeshAttribute::VertexInstance::Tangent);
-	TVertexInstanceAttributesRef<float> TargetVertexInstanceBinormalSigns = TargetMesh.VertexInstanceAttributes().GetAttributesRef<float>(MeshAttribute::VertexInstance::BinormalSign);
-	TVertexInstanceAttributesRef<FVector4> TargetVertexInstanceColors = TargetMesh.VertexInstanceAttributes().GetAttributesRef<FVector4>(MeshAttribute::VertexInstance::Color);
-	TVertexInstanceAttributesRef<FVector2D> TargetVertexInstanceUVs = TargetMesh.VertexInstanceAttributes().GetAttributesRef<FVector2D>(MeshAttribute::VertexInstance::TextureCoordinate);
+	FStaticMeshConstAttributes SourceAttributes(SourceMesh);
+	TVertexAttributesConstRef<FVector> SourceVertexPositions = SourceAttributes.GetVertexPositions();
+	TEdgeAttributesConstRef<bool> SourceEdgeHardnesses = SourceAttributes.GetEdgeHardnesses();
+	TPolygonGroupAttributesConstRef<FName> SourceImportedMaterialSlotNames = SourceAttributes.GetPolygonGroupMaterialSlotNames();
+	TVertexInstanceAttributesConstRef<FVector> SourceVertexInstanceNormals = SourceAttributes.GetVertexInstanceNormals();
+	TVertexInstanceAttributesConstRef<FVector> SourceVertexInstanceTangents = SourceAttributes.GetVertexInstanceTangents();
+	TVertexInstanceAttributesConstRef<float> SourceVertexInstanceBinormalSigns = SourceAttributes.GetVertexInstanceBinormalSigns();
+	TVertexInstanceAttributesConstRef<FVector4> SourceVertexInstanceColors = SourceAttributes.GetVertexInstanceColors();
+	TVertexInstanceAttributesConstRef<FVector2D> SourceVertexInstanceUVs = SourceAttributes.GetVertexInstanceUVs();
+
+	FStaticMeshAttributes TargetAttributes(TargetMesh);
+	TVertexAttributesRef<FVector> TargetVertexPositions = TargetAttributes.GetVertexPositions();
+	TEdgeAttributesRef<bool> TargetEdgeHardnesses = TargetAttributes.GetEdgeHardnesses();
+	TPolygonGroupAttributesRef<FName> TargetImportedMaterialSlotNames = TargetAttributes.GetPolygonGroupMaterialSlotNames();
+	TVertexInstanceAttributesRef<FVector> TargetVertexInstanceNormals = TargetAttributes.GetVertexInstanceNormals();
+	TVertexInstanceAttributesRef<FVector> TargetVertexInstanceTangents = TargetAttributes.GetVertexInstanceTangents();
+	TVertexInstanceAttributesRef<float> TargetVertexInstanceBinormalSigns = TargetAttributes.GetVertexInstanceBinormalSigns();
+	TVertexInstanceAttributesRef<FVector4> TargetVertexInstanceColors = TargetAttributes.GetVertexInstanceColors();
+	TVertexInstanceAttributesRef<FVector2D> TargetVertexInstanceUVs = TargetAttributes.GetVertexInstanceUVs();
 
 	//Copy into the target mesh
 	TargetMesh.ReserveNewVertices(SourceMesh.Vertices().Num());
@@ -854,7 +969,7 @@ void FStaticMeshOperations::AppendMeshDescription(const FMeshDescription& Source
 	TargetMesh.ReserveNewPolygons(SourceMesh.Polygons().Num());
 
 	int32 NumSourceUVChannels = 0;
-	for (int32 ChannelIdx = 0; ChannelIdx < SourceVertexInstanceUVs.GetNumIndices(); ++ChannelIdx)
+	for (int32 ChannelIdx = 0; ChannelIdx < SourceVertexInstanceUVs.GetNumChannels(); ++ChannelIdx)
 	{
 		if (AppendSettings.bMergeUVChannels[ChannelIdx])
 		{
@@ -862,9 +977,9 @@ void FStaticMeshOperations::AppendMeshDescription(const FMeshDescription& Source
 		}
 	}
 
-	if (NumSourceUVChannels > TargetVertexInstanceUVs.GetNumIndices())
+	if (NumSourceUVChannels > TargetVertexInstanceUVs.GetNumChannels())
 	{
-		TargetVertexInstanceUVs.SetNumIndices(NumSourceUVChannels);
+		TargetVertexInstanceUVs.SetNumChannels(NumSourceUVChannels);
 	}
 
 	//PolygonGroups
@@ -877,7 +992,7 @@ void FStaticMeshOperations::AppendMeshDescription(const FMeshDescription& Source
 	{
 		for (FPolygonGroupID SourcePolygonGroupID : SourceMesh.PolygonGroups().GetElementIDs())
 		{
-			FPolygonGroupID TargetMatchingID = FPolygonGroupID::Invalid;
+			FPolygonGroupID TargetMatchingID = INDEX_NONE;
 			for (FPolygonGroupID TargetPolygonGroupID : TargetMesh.PolygonGroups().GetElementIDs())
 			{
 				if (SourceImportedMaterialSlotNames[SourcePolygonGroupID] == TargetImportedMaterialSlotNames[TargetPolygonGroupID])
@@ -886,7 +1001,7 @@ void FStaticMeshOperations::AppendMeshDescription(const FMeshDescription& Source
 					break;
 				}
 			}
-			if (TargetMatchingID == FPolygonGroupID::Invalid)
+			if (TargetMatchingID == INDEX_NONE)
 			{
 				TargetMatchingID = TargetMesh.CreatePolygonGroup();
 				TargetImportedMaterialSlotNames[TargetMatchingID] = SourceImportedMaterialSlotNames[SourcePolygonGroupID];
@@ -902,7 +1017,6 @@ void FStaticMeshOperations::AppendMeshDescription(const FMeshDescription& Source
 	{
 		FVertexID TargetVertexID = TargetMesh.CreateVertex();
 		TargetVertexPositions[TargetVertexID] = (SourceVertexPositions[SourceVertexID] - AppendSettings.MergedAssetPivot);
-		TargetVertexCornerSharpness[TargetVertexID] = SourceVertexCornerSharpness[SourceVertexID];
 		SourceVertexIDRemap.Add(SourceVertexID, TargetVertexID);
 	}
 
@@ -926,7 +1040,6 @@ void FStaticMeshOperations::AppendMeshDescription(const FMeshDescription& Source
 		const FVertexID EdgeVertex1 = SourceMesh.GetEdgeVertex(SourceEdgeID, 1);
 		FEdgeID TargetEdgeID = TargetMesh.CreateEdge(SourceVertexIDRemap[EdgeVertex0], SourceVertexIDRemap[EdgeVertex1]);
 		TargetEdgeHardnesses[TargetEdgeID] = SourceEdgeHardnesses[SourceEdgeID];
-		TargetEdgeCreaseSharpnesses[TargetEdgeID] = SourceEdgeCreaseSharpnesses[SourceEdgeID];
 		SourceEdgeIDRemap.Add(SourceEdgeID, TargetEdgeID);
 	}
 
@@ -976,7 +1089,7 @@ void FStaticMeshOperations::AppendMeshDescription(const FMeshDescription& Source
 	//Polygons
 	for (const FPolygonID SourcePolygonID : SourceMesh.Polygons().GetElementIDs())
 	{
-		const TArray<FVertexInstanceID>& PerimeterVertexInstanceIDs = SourceMesh.GetPolygonVertexInstances(SourcePolygonID);
+		TArray<FVertexInstanceID, TInlineAllocator<4>> PerimeterVertexInstanceIDs = SourceMesh.GetPolygonVertexInstances<TInlineAllocator<4>>(SourcePolygonID);
 		const FPolygonGroupID PolygonGroupID = SourceMesh.GetPolygonPolygonGroup(SourcePolygonID);
 		//Find the polygonGroupID
 		FPolygonGroupID TargetPolygonGroupID = RemapPolygonGroup[PolygonGroupID];
@@ -1001,10 +1114,12 @@ void FStaticMeshOperations::AreNormalsAndTangentsValid(const FMeshDescription& M
 {
 	bHasInvalidNormals = false;
 	bHasInvalidTangents = false;
-	TVertexInstanceAttributesConstRef<FVector> VertexInstanceNormals = MeshDescription.VertexInstanceAttributes().GetAttributesRef<FVector>(MeshAttribute::VertexInstance::Normal);
-	TVertexInstanceAttributesConstRef<FVector> VertexInstanceTangents = MeshDescription.VertexInstanceAttributes().GetAttributesRef<FVector>(MeshAttribute::VertexInstance::Tangent);
 
-	for (const FVertexInstanceID& VertexInstanceID : MeshDescription.VertexInstances().GetElementIDs())
+	FStaticMeshConstAttributes Attributes(MeshDescription);
+	TArrayView<const FVector> VertexInstanceNormals = Attributes.GetVertexInstanceNormals().GetRawArray();
+	TArrayView<const FVector> VertexInstanceTangents = Attributes.GetVertexInstanceTangents().GetRawArray();
+
+	for (const FVertexInstanceID VertexInstanceID : MeshDescription.VertexInstances().GetElementIDs())
 	{
 		bHasInvalidNormals |= (VertexInstanceNormals[VertexInstanceID].IsNearlyZero() || VertexInstanceNormals[VertexInstanceID].ContainsNaN());
 		bHasInvalidTangents |= (VertexInstanceTangents[VertexInstanceID].IsNearlyZero() || VertexInstanceTangents[VertexInstanceID].ContainsNaN());
@@ -1022,12 +1137,13 @@ void ClearNormalsAndTangentsData(FMeshDescription& MeshDescription, bool bClearN
 		return;
 	}
 
-	TVertexInstanceAttributesRef<FVector> VertexInstanceNormals = MeshDescription.VertexInstanceAttributes().GetAttributesRef<FVector>(MeshAttribute::VertexInstance::Normal);
-	TVertexInstanceAttributesRef<FVector> VertexInstanceTangents = MeshDescription.VertexInstanceAttributes().GetAttributesRef<FVector>(MeshAttribute::VertexInstance::Tangent);
-	TVertexInstanceAttributesRef<float> VertexBinormalSigns = MeshDescription.VertexInstanceAttributes().GetAttributesRef<float>(MeshAttribute::VertexInstance::BinormalSign);
+	FStaticMeshAttributes Attributes(MeshDescription);
+	TArrayView<FVector> VertexInstanceNormals = Attributes.GetVertexInstanceNormals().GetRawArray();
+	TArrayView<FVector> VertexInstanceTangents = Attributes.GetVertexInstanceTangents().GetRawArray();
+	TArrayView<float> VertexInstanceBinormals = Attributes.GetVertexInstanceBinormalSigns().GetRawArray();
 
 	//Zero out all value that need to be recompute
-	for (const FVertexInstanceID& VertexInstanceID : MeshDescription.VertexInstances().GetElementIDs())
+	for (const FVertexInstanceID VertexInstanceID : MeshDescription.VertexInstances().GetElementIDs())
 	{
 		if (bClearNormals)
 		{
@@ -1035,8 +1151,8 @@ void ClearNormalsAndTangentsData(FMeshDescription& MeshDescription, bool bClearN
 		}
 		if (bClearTangents)
 		{
-			//Dump the tangents
-			VertexBinormalSigns[VertexInstanceID] = 0.0f;
+			// Dump the tangents
+			VertexInstanceBinormals[VertexInstanceID] = 0.0f;
 			VertexInstanceTangents[VertexInstanceID] = FVector::ZeroVector;
 		}
 	}
@@ -1134,13 +1250,14 @@ void FStaticMeshOperations::ComputeTangentsAndNormals(FMeshDescription& MeshDesc
 		FTriangleCornerData CornerAngleDatas[3];
 	};
 
-
+#if 0
 	//Make sure the meshdescription is triangulate
 	if (MeshDescription.Triangles().Num() < MeshDescription.Polygons().Num())
 	{
 		//Triangulate the mesh, we compute the normals on triangle not on polygon.
 		MeshDescription.TriangulateMesh();
 	}
+#endif
 
 	const bool bForceComputeNormals = EnumHasAllFlags(ComputeNTBsOptions, EComputeNTBsFlags::Normals);
 	const bool bForceComputeTangent = EnumHasAnyFlags(ComputeNTBsOptions, EComputeNTBsFlags::Normals | EComputeNTBsFlags::Tangents);
@@ -1150,58 +1267,66 @@ void FStaticMeshOperations::ComputeTangentsAndNormals(FMeshDescription& MeshDesc
 	//Clear any data we want to force-recompute since the following code actually look for any invalid data and recompute it.
 	ClearNormalsAndTangentsData(MeshDescription, bForceComputeNormals, bForceComputeTangent);
 
-	//Compute the weight (area and angle) for each triangles
-	TMap<FTriangleID, FTriangleData> TriangleDatas;
+	// Going to iterate over all triangles, so mandate that the triangle elements are compact, i.e. there are no holes
+	const int32 NumTriangles = MeshDescription.Triangles().Num();
+	check(MeshDescription.Triangles().GetArraySize() == NumTriangles);
+
+	//vCompute the weight (area and angle) for each triangles
+	TArray<FTriangleData> TriangleDatas;
+	TriangleDatas.SetNum(NumTriangles);
 	if (bComputeWeightedNormals)
 	{
-		TVertexAttributesConstRef<FVector> VertexPositions = MeshDescription.VertexAttributes().GetAttributesRef<FVector>(MeshAttribute::Vertex::Position);
-		TriangleDatas.Reserve(MeshDescription.Triangles().Num());
+		FStaticMeshAttributes Attributes(MeshDescription);
+		TArrayView<const FVector> VertexPositions = Attributes.GetVertexPositions().GetRawArray();
+		TArrayView<const FVertexID> TriVertexIDs = Attributes.GetTriangleVertexIndices().GetRawArray();
+		TArrayView<const FVertexInstanceID> TriVertexInstanceIDs = Attributes.GetTriangleVertexInstanceIndices().GetRawArray();
+
+		TriangleDatas.Reserve(NumTriangles);
 
-		for (FTriangleID TriangleID : MeshDescription.Triangles().GetElementIDs())
+		for (int32 Index = 0, TriIndex = 0; Index < NumTriangles; Index++, TriIndex += 3)
 		{
-			TArrayView<const FVertexInstanceID> VertexInstanceIDs = MeshDescription.GetTriangleVertexInstances(TriangleID);
-			//Triangle should use 3 vertex instances
-			check(VertexInstanceIDs.Num() == 3);
-			const FVector& PointA = VertexPositions[MeshDescription.GetVertexInstanceVertex(VertexInstanceIDs[0])];
-			const FVector& PointB = VertexPositions[MeshDescription.GetVertexInstanceVertex(VertexInstanceIDs[1])];
-			const FVector& PointC = VertexPositions[MeshDescription.GetVertexInstanceVertex(VertexInstanceIDs[2])];
-			FTriangleData& TriangleData = TriangleDatas.FindOrAdd(TriangleID);
+			const FVector& PointA = VertexPositions[TriVertexIDs[TriIndex + 0]];
+			const FVector& PointB = VertexPositions[TriVertexIDs[TriIndex + 1]];
+			const FVector& PointC = VertexPositions[TriVertexIDs[TriIndex + 2]];
+			FTriangleData& TriangleData = TriangleDatas[Index];
 			TriangleData.Area = TriangleUtilities::ComputeTriangleArea(PointA, PointB, PointC);
-			TriangleData.SetCornerAngleData(VertexInstanceIDs[0], TriangleUtilities::ComputeTriangleCornerAngle(PointA, PointB, PointC), 0);
-			TriangleData.SetCornerAngleData(VertexInstanceIDs[1], TriangleUtilities::ComputeTriangleCornerAngle(PointB, PointC, PointA), 1);
-			TriangleData.SetCornerAngleData(VertexInstanceIDs[2], TriangleUtilities::ComputeTriangleCornerAngle(PointC, PointA, PointB), 2);
+			TriangleData.SetCornerAngleData(TriVertexInstanceIDs[TriIndex + 0], TriangleUtilities::ComputeTriangleCornerAngle(PointA, PointB, PointC), 0);
+			TriangleData.SetCornerAngleData(TriVertexInstanceIDs[TriIndex + 1], TriangleUtilities::ComputeTriangleCornerAngle(PointB, PointC, PointA), 1);
+			TriangleData.SetCornerAngleData(TriVertexInstanceIDs[TriIndex + 2], TriangleUtilities::ComputeTriangleCornerAngle(PointC, PointA, PointB), 2);
 		}
 	}
 
-	//Iterate all vertex to compute normals for all vertex instance
-	TArray<FVertexID> Vertices;
-	Vertices.Reserve(MeshDescription.Vertices().Num());
-	for (const FVertexID VertexID : MeshDescription.Vertices().GetElementIDs())
-	{
-		Vertices.Add(VertexID);
-	}
+	// Ensure certain indexers are built in anticipation
+	MeshDescription.BuildVertexIndexers();
+	MeshDescription.BuildEdgeIndexers();
+
+	// Going to iterate over all vertices, so mandate that the vertex elements are compact, i.e. there are no holes
+	const int32 NumVertices = MeshDescription.Vertices().Num();
+	check(MeshDescription.Vertices().GetArraySize() == NumVertices);
 
 	// Split work in batch to reduce call and allocation overhead
 	const int32 BatchSize = 128 * 1024;
-	const int32 BatchCount = 1 + Vertices.Num() / BatchSize;
+	const int32 BatchCount = (NumVertices + BatchSize - 1) / BatchSize;
 
 	//Iterate all vertex to compute normals for all vertex instance
 	ParallelFor(BatchCount,
-		[&Vertices, &BatchSize, &bComputeTangentWithMikkTSpace, &MeshDescription, bComputeWeightedNormals, &TriangleDatas](int32 BatchIndex)
+		[NumVertices, BatchSize, bComputeTangentWithMikkTSpace, bComputeWeightedNormals, &MeshDescription, &TriangleDatas](int32 BatchIndex)
 		{
-			TVertexInstanceAttributesConstRef<FVector2D> VertexUVs = MeshDescription.VertexInstanceAttributes().GetAttributesRef<FVector2D>(MeshAttribute::VertexInstance::TextureCoordinate);
-			TVertexInstanceAttributesRef<FVector> VertexNormals = MeshDescription.VertexInstanceAttributes().GetAttributesRef<FVector>(MeshAttribute::VertexInstance::Normal);
-			TVertexInstanceAttributesRef<FVector> VertexTangents = MeshDescription.VertexInstanceAttributes().GetAttributesRef<FVector>(MeshAttribute::VertexInstance::Tangent);
-			TVertexInstanceAttributesRef<float> VertexBinormalSigns = MeshDescription.VertexInstanceAttributes().GetAttributesRef<float>(MeshAttribute::VertexInstance::BinormalSign);
+			FStaticMeshAttributes Attributes(MeshDescription);
 
-			TPolygonAttributesConstRef<FVector> PolygonNormals = MeshDescription.PolygonAttributes().GetAttributesRef<FVector>(MeshAttribute::Polygon::Normal);
-			TPolygonAttributesConstRef<FVector> PolygonTangents = MeshDescription.PolygonAttributes().GetAttributesRef<FVector>(MeshAttribute::Polygon::Tangent);
-			TPolygonAttributesConstRef<FVector> PolygonBinormals = MeshDescription.PolygonAttributes().GetAttributesRef<FVector>(MeshAttribute::Polygon::Binormal);
-			TEdgeAttributesConstRef<bool>       EdgeHardnesses = MeshDescription.EdgeAttributes().GetAttributesRef<bool>(MeshAttribute::Edge::IsHard);
+			TArrayView<const FVector2D> VertexUVs = Attributes.GetVertexInstanceUVs().GetRawArray(0);	// Use UV0
+			TArrayView<const FVector> TriangleNormals = Attributes.GetTriangleNormals().GetRawArray();
+			TArrayView<const FVector> TriangleTangents = Attributes.GetTriangleTangents().GetRawArray();
+			TArrayView<const FVector> TriangleBinormals = Attributes.GetTriangleBinormals().GetRawArray();
+			TArrayView<const bool> EdgeHardnesses = Attributes.GetEdgeHardnesses().GetRawArray();
 
-			check(PolygonNormals.IsValid());
-			check(PolygonTangents.IsValid());
-			check(PolygonBinormals.IsValid());
+			TArrayView<FVector> VertexNormals = Attributes.GetVertexInstanceNormals().GetRawArray();
+			TArrayView<FVector> VertexTangents = Attributes.GetVertexInstanceTangents().GetRawArray();
+			TArrayView<float> VertexBinormalSigns = Attributes.GetVertexInstanceBinormalSigns().GetRawArray();
+
+			check(TriangleNormals.Num() > 0);
+			check(TriangleTangents.Num() > 0);
+			check(TriangleBinormals.Num() > 0);
 
 			//Reuse containers between iterations to reduce allocations
 			TMap<FVector2D, FVector, FDefaultSetAllocator, FNTBGroupKeyFuncs> GroupTangent;
@@ -1214,23 +1339,23 @@ void FStaticMeshOperations::ComputeTangentsAndNormals(FMeshDescription& MeshDesc
 
 			VertexInfoMap.Reserve(20);
 
-			int32 Indice = BatchIndex * BatchSize;
-			int32 LastIndice = FMath::Min(Indice + BatchSize, Vertices.Num());
-			for (; Indice < LastIndice; ++Indice)
+			int32 StartIndex = BatchIndex * BatchSize;
+			int32 LastIndex = FMath::Min(StartIndex + BatchSize, NumVertices);
+			for (int32 Index = StartIndex; Index < LastIndex; ++Index)
 			{
 				VertexInfoMap.Reset();
 
-				const FVertexID VertexID = Vertices[Indice];
+				const FVertexID VertexID(Index);
 
 				bool bPointHasAllTangents = true;
-				//Fill the VertexInfoMap
+				// Fill the VertexInfoMap
 				for (const FEdgeID EdgeID : MeshDescription.GetVertexConnectedEdges(VertexID))
 				{
 					for (const FTriangleID TriangleID : MeshDescription.GetEdgeConnectedTriangles(EdgeID))
 					{
 						FVertexInfo& VertexInfo = VertexInfoMap.FindOrAdd(TriangleID);
 						int32 EdgeIndex = VertexInfo.EdgeIDs.AddUnique(EdgeID);
-						if (VertexInfo.TriangleID == FTriangleID::Invalid)
+						if (VertexInfo.TriangleID == INDEX_NONE)
 						{
 							VertexInfo.TriangleID = TriangleID;
 							for (const FVertexInstanceID VertexInstanceID : MeshDescription.GetTriangleVertexInstances(TriangleID))
@@ -1238,7 +1363,7 @@ void FStaticMeshOperations::ComputeTangentsAndNormals(FMeshDescription& MeshDesc
 								if (MeshDescription.GetVertexInstanceVertex(VertexInstanceID) == VertexID)
 								{
 									VertexInfo.VertexInstanceID = VertexInstanceID;
-									VertexInfo.UVs = VertexUVs.Get(VertexInstanceID, 0);	// UV0
+									VertexInfo.UVs = VertexUVs[VertexInstanceID];	// UV0
 									bPointHasAllTangents &= !VertexNormals[VertexInstanceID].IsNearlyZero() && !VertexTangents[VertexInstanceID].IsNearlyZero();
 									if (bPointHasAllTangents)
 									{
@@ -1319,24 +1444,23 @@ void FStaticMeshOperations::ComputeTangentsAndNormals(FMeshDescription& MeshDesc
 					FVector GroupNormal(FVector::ZeroVector);
 					for (const FTriangleID TriangleID : Group)
 					{
-						FPolygonID PolygonID = MeshDescription.GetTrianglePolygon(TriangleID);
 						FVertexInfo& CurrentVertexInfo = VertexInfoMap.FindOrAdd(TriangleID);
 						float CornerWeight = 1.0f;
 
 						if (bComputeWeightedNormals)
 						{
-							FTriangleData& TriangleData = TriangleDatas.FindChecked(TriangleID);
+							FTriangleData& TriangleData = TriangleDatas[TriangleID];
 							CornerWeight = TriangleData.Area * TriangleData.GetCornerAngle(CurrentVertexInfo.VertexInstanceID);
 						}
 
-						const FVector PolyNormal = CornerWeight * PolygonNormals[PolygonID];
-						const FVector PolyTangent = CornerWeight * PolygonTangents[PolygonID];
-						const FVector PolyBinormal = CornerWeight * PolygonBinormals[PolygonID];
+						const FVector TriNormal = CornerWeight * TriangleNormals[TriangleID];
+						const FVector TriTangent = CornerWeight * TriangleTangents[TriangleID];
+						const FVector TriBinormal = CornerWeight * TriangleBinormals[TriangleID];
 
 						VertexInstanceInGroup.Add(VertexInfoMap[TriangleID].VertexInstanceID);
-						if (!PolyNormal.IsNearlyZero(SMALL_NUMBER) && !PolyNormal.ContainsNaN())
+						if (!TriNormal.IsNearlyZero(SMALL_NUMBER) && !TriNormal.ContainsNaN())
 						{
-							GroupNormal += PolyNormal;
+							GroupNormal += TriNormal;
 						}
 						if (!bComputeTangentWithMikkTSpace)
 						{
@@ -1349,13 +1473,13 @@ void FStaticMeshOperations::ComputeTangentsAndNormals(FMeshDescription& MeshDesc
 								GroupTangentValue = FVector(0.0f);
 								GroupBiNormalValue = FVector(0.0f);
 							}
-							if (!PolyTangent.IsNearlyZero(SMALL_NUMBER) && !PolyTangent.ContainsNaN())
+							if (!TriTangent.IsNearlyZero(SMALL_NUMBER) && !TriTangent.ContainsNaN())
 							{
-								GroupTangentValue += PolyTangent;
+								GroupTangentValue += TriTangent;
 							}
-							if (!PolyBinormal.IsNearlyZero(SMALL_NUMBER) && !PolyBinormal.ContainsNaN())
+							if (!TriBinormal.IsNearlyZero(SMALL_NUMBER) && !TriBinormal.ContainsNaN())
 							{
-								GroupBiNormalValue += PolyBinormal;
+								GroupBiNormalValue += TriBinormal;
 							}
 						}
 					}
@@ -1377,7 +1501,7 @@ void FStaticMeshOperations::ComputeTangentsAndNormals(FMeshDescription& MeshDesc
 					//Apply the average NTB on all Vertex instance
 					for (const FVertexInstanceID VertexInstanceID : VertexInstanceInGroup)
 					{
-						const FVector2D& VertexUV = VertexUVs.Get(VertexInstanceID, 0);	// UV0
+						const FVector2D& VertexUV = VertexUVs[VertexInstanceID];	// UV0
 
 						if (VertexNormals[VertexInstanceID].IsNearlyZero(SMALL_NUMBER))
 						{
@@ -1434,31 +1558,34 @@ void FStaticMeshOperations::ComputeTangentsAndNormals(FMeshDescription& MeshDesc
 #if WITH_MIKKTSPACE
 namespace MeshDescriptionMikktSpaceInterface
 {
+	struct FMeshDescriptionCachedData
+	{
+		int32 NumTriangles;
+		TArrayView<const FVertexID> TriangleVertexIDs;
+		TArrayView<const FVertexInstanceID> TriangleVertexInstanceIDs;
+		TArrayView<const FVector> VertexPositions;
+		TArrayView<const FVector> VertexInstanceNormals;
+		TArrayView<const FVector2D> VertexInstanceUVs;
+		TArrayView<FVector> VertexInstanceTangents;
+		TArrayView<float> VertexInstanceBinormalSigns;
+	};
+
 	int MikkGetNumFaces(const SMikkTSpaceContext* Context)
 	{
-		FMeshDescription *MeshDescription = (FMeshDescription*)(Context->m_pUserData);
-		return MeshDescription->Polygons().GetArraySize();
+		FMeshDescriptionCachedData* UserData = (FMeshDescriptionCachedData*)(Context->m_pUserData);
+		return UserData->NumTriangles;
 	}
 
 	int MikkGetNumVertsOfFace(const SMikkTSpaceContext* Context, const int FaceIdx)
 	{
-		// All of our meshes are triangles.
-		FMeshDescription *MeshDescription = (FMeshDescription*)(Context->m_pUserData);
-		if (MeshDescription->IsPolygonValid(FPolygonID(FaceIdx)))
-		{
-			return MeshDescription->GetPolygonVertexInstances(FPolygonID(FaceIdx)).Num();
-		}
-
-		return 0;
+		FMeshDescriptionCachedData* UserData = (FMeshDescriptionCachedData*)(Context->m_pUserData);
+		return 3;
 	}
 
 	void MikkGetPosition(const SMikkTSpaceContext* Context, float Position[3], const int FaceIdx, const int VertIdx)
 	{
-		FMeshDescription* MeshDescription = (FMeshDescription*)(Context->m_pUserData);
-		const TArray<FVertexInstanceID>& VertexInstanceIDs = MeshDescription->GetPolygonVertexInstances(FPolygonID(FaceIdx));
-		const FVertexInstanceID VertexInstanceID = VertexInstanceIDs[VertIdx];
-		const FVertexID VertexID = MeshDescription->GetVertexInstanceVertex(VertexInstanceID);
-		const FVector& VertexPosition = MeshDescription->VertexAttributes().GetAttribute<FVector>(VertexID, MeshAttribute::Vertex::Position);
+		FMeshDescriptionCachedData* UserData = (FMeshDescriptionCachedData*)(Context->m_pUserData);
+		const FVector& VertexPosition = UserData->VertexPositions[UserData->TriangleVertexIDs[FaceIdx * 3 + VertIdx]];
 		Position[0] = VertexPosition.X;
 		Position[1] = VertexPosition.Y;
 		Position[2] = VertexPosition.Z;
@@ -1466,10 +1593,8 @@ namespace MeshDescriptionMikktSpaceInterface
 
 	void MikkGetNormal(const SMikkTSpaceContext* Context, float Normal[3], const int FaceIdx, const int VertIdx)
 	{
-		FMeshDescription* MeshDescription = (FMeshDescription*)(Context->m_pUserData);
-		const TArray<FVertexInstanceID>& VertexInstanceIDs = MeshDescription->GetPolygonVertexInstances(FPolygonID(FaceIdx));
-		const FVertexInstanceID VertexInstanceID = VertexInstanceIDs[VertIdx];
-		const FVector& VertexNormal = MeshDescription->VertexInstanceAttributes().GetAttribute<FVector>(VertexInstanceID, MeshAttribute::VertexInstance::Normal);
+		FMeshDescriptionCachedData* UserData = (FMeshDescriptionCachedData*)(Context->m_pUserData);
+		const FVector& VertexNormal = UserData->VertexInstanceNormals[UserData->TriangleVertexInstanceIDs[FaceIdx * 3 + VertIdx]];
 		Normal[0] = VertexNormal.X;
 		Normal[1] = VertexNormal.Y;
 		Normal[2] = VertexNormal.Z;
@@ -1477,20 +1602,16 @@ namespace MeshDescriptionMikktSpaceInterface
 
 	void MikkSetTSpaceBasic(const SMikkTSpaceContext* Context, const float Tangent[3], const float BitangentSign, const int FaceIdx, const int VertIdx)
 	{
-		FMeshDescription* MeshDescription = (FMeshDescription*)(Context->m_pUserData);
-		const TArray<FVertexInstanceID>& VertexInstanceIDs = MeshDescription->GetPolygonVertexInstances(FPolygonID(FaceIdx));
-		const FVertexInstanceID VertexInstanceID = VertexInstanceIDs[VertIdx];
-		const FVector VertexTangent(Tangent[0], Tangent[1], Tangent[2]);
-		MeshDescription->VertexInstanceAttributes().SetAttribute<FVector>(VertexInstanceID, MeshAttribute::VertexInstance::Tangent, 0, VertexTangent);
-		MeshDescription->VertexInstanceAttributes().SetAttribute<float>(VertexInstanceID, MeshAttribute::VertexInstance::BinormalSign, 0, -BitangentSign);
+		FMeshDescriptionCachedData* UserData = (FMeshDescriptionCachedData*)(Context->m_pUserData);
+		const FVertexInstanceID VertexInstanceID = UserData->TriangleVertexInstanceIDs[FaceIdx * 3 + VertIdx];
+		UserData->VertexInstanceTangents[VertexInstanceID] = FVector(Tangent[0], Tangent[1], Tangent[2]);
+		UserData->VertexInstanceBinormalSigns[VertexInstanceID] = -BitangentSign;
 	}
 
 	void MikkGetTexCoord(const SMikkTSpaceContext* Context, float UV[2], const int FaceIdx, const int VertIdx)
 	{
-		FMeshDescription* MeshDescription = (FMeshDescription*)(Context->m_pUserData);
-		const TArray<FVertexInstanceID>& VertexInstanceIDs = MeshDescription->GetPolygonVertexInstances(FPolygonID(FaceIdx));
-		const FVertexInstanceID VertexInstanceID = VertexInstanceIDs[VertIdx];
-		const FVector2D& TexCoord = MeshDescription->VertexInstanceAttributes().GetAttribute<FVector2D>(VertexInstanceID, MeshAttribute::VertexInstance::TextureCoordinate, 0);
+		FMeshDescriptionCachedData* UserData = (FMeshDescriptionCachedData*)(Context->m_pUserData);
+		const FVector2D& TexCoord = UserData->VertexInstanceUVs[UserData->TriangleVertexInstanceIDs[FaceIdx * 3 + VertIdx]];
 		UV[0] = TexCoord.X;
 		UV[1] = TexCoord.Y;
 	}
@@ -1504,7 +1625,7 @@ void FStaticMeshOperations::ComputeMikktTangents(FMeshDescription& MeshDescripti
 
 	// The Mikkt interface does not handle properly polygon array with 'holes'
 	// Compact mesh description if this is the case
-	if (MeshDescription.Polygons().Num() != MeshDescription.Polygons().GetArraySize())
+	if (MeshDescription.Triangles().Num() != MeshDescription.Triangles().GetArraySize())
 	{
 		FElementIDRemappings Remappings;
 		MeshDescription.Compact(Remappings);
@@ -1520,9 +1641,21 @@ void FStaticMeshOperations::ComputeMikktTangents(FMeshDescription& MeshDescripti
 	MikkTInterface.m_setTSpaceBasic = MeshDescriptionMikktSpaceInterface::MikkSetTSpaceBasic;
 	MikkTInterface.m_setTSpace = nullptr;
 
+	MeshDescriptionMikktSpaceInterface::FMeshDescriptionCachedData UserData;
+	UserData.NumTriangles = MeshDescription.Triangles().Num();
+
+	FStaticMeshAttributes Attributes(MeshDescription);
+	UserData.TriangleVertexIDs = Attributes.GetTriangleVertexIndices().GetRawArray();
+	UserData.TriangleVertexInstanceIDs = Attributes.GetTriangleVertexInstanceIndices().GetRawArray();
+	UserData.VertexPositions = Attributes.GetVertexPositions().GetRawArray();
+	UserData.VertexInstanceUVs = Attributes.GetVertexInstanceUVs().GetRawArray(0);
+	UserData.VertexInstanceNormals = Attributes.GetVertexInstanceNormals().GetRawArray();
+	UserData.VertexInstanceTangents = Attributes.GetVertexInstanceTangents().GetRawArray();
+	UserData.VertexInstanceBinormalSigns = Attributes.GetVertexInstanceBinormalSigns().GetRawArray();
+	
 	SMikkTSpaceContext MikkTContext;
 	MikkTContext.m_pInterface = &MikkTInterface;
-	MikkTContext.m_pUserData = (void*)(&MeshDescription);
+	MikkTContext.m_pUserData = (void*)(&UserData);
 	MikkTContext.m_bIgnoreDegenerates = bIgnoreDegenerateTriangles;
 	genTangSpaceDefault(&MikkTContext);
 #else
@@ -1553,7 +1686,7 @@ void FStaticMeshOperations::FindOverlappingCorners(FOverlappingCorners& OutOverl
 	int32 WedgeIndex = 0;
 	for (const FPolygonID PolygonID : MeshDescription.Polygons().GetElementIDs())
 	{
-		const TArray<FTriangleID>& TriangleIDs = MeshDescription.GetPolygonTriangleIDs(PolygonID);
+		TArrayView<const FTriangleID> TriangleIDs = MeshDescription.GetPolygonTriangleIDs(PolygonID);
 		for (const FTriangleID TriangleID : TriangleIDs)
 		{
 			for (int32 Corner = 0; Corner < 3; ++Corner)
@@ -1624,7 +1757,7 @@ struct FLayoutUVMeshDescriptionView final : FLayoutUV::IMeshView
 
 		for (const FPolygonID PolygonID : MeshDescription.Polygons().GetElementIDs())
 		{
-			const TArray<FTriangleID>& TriangleIDs = MeshDescription.GetPolygonTriangleIDs(PolygonID);
+			TArrayView<const FTriangleID> TriangleIDs = MeshDescription.GetPolygonTriangleIDs(PolygonID);
 			for (const FTriangleID TriangleID : TriangleIDs)
 			{
 				for (int32 Corner = 0; Corner < 3; ++Corner)
@@ -1662,10 +1795,10 @@ struct FLayoutUVMeshDescriptionView final : FLayoutUV::IMeshView
 	void InitOutputTexcoords(uint32 Num) override
 	{
 		// If current DstChannel is out of range of the number of UVs defined by the mesh description, change the index count accordingly
-		const uint32 NumUVs = TexCoords.GetNumIndices();
+		const uint32 NumUVs = TexCoords.GetNumChannels();
 		if (DstChannel >= NumUVs)
 		{
-			TexCoords.SetNumIndices(DstChannel + 1);
+			TexCoords.SetNumChannels(DstChannel + 1);
 			ensure(false);	// not expecting it to get here
 		}
 	}
@@ -1722,9 +1855,9 @@ bool FStaticMeshOperations::GenerateUniqueUVsForStaticMesh(const FMeshDescriptio
 	//Make sure we have a destination UV TextureCoordinnate
 	{
 		TVertexInstanceAttributesRef<FVector2D> DuplicateVertexInstanceUVs = DuplicateMeshDescription.VertexInstanceAttributes().GetAttributesRef<FVector2D>(MeshAttribute::VertexInstance::TextureCoordinate);
-		if (DuplicateVertexInstanceUVs.GetNumIndices() < 2)
+		if (DuplicateVertexInstanceUVs.GetNumChannels() < 2)
 		{
-			DuplicateVertexInstanceUVs.SetNumIndices(2);
+			DuplicateVertexInstanceUVs.SetNumChannels(2);
 		}
 	}
 
@@ -1744,7 +1877,7 @@ bool FStaticMeshOperations::GenerateUniqueUVsForStaticMesh(const FMeshDescriptio
 		for (FPolygonID RefPolygonID : DuplicateMeshDescription.Polygons().GetElementIDs())
 		{
 			const FPolygonGroupID RefPolygonGroupID = DuplicateMeshDescription.GetPolygonPolygonGroup(RefPolygonID);
-			const TArray<FVertexInstanceID>& RefVertexInstances = DuplicateMeshDescription.GetPolygonVertexInstances(RefPolygonID);
+			TArray<FVertexInstanceID, TInlineAllocator<4>> RefVertexInstances = DuplicateMeshDescription.GetPolygonVertexInstances<TInlineAllocator<4>>(RefPolygonID);
 
 			int32 PolyHash = GetTypeHash(RefPolygonGroupID);
 			for (FVertexInstanceID RefVertexInstanceID : RefVertexInstances)
@@ -1753,7 +1886,7 @@ bool FStaticMeshOperations::GenerateUniqueUVsForStaticMesh(const FMeshDescriptio
 				DuplicateMeshDescription.VertexInstanceAttributes().ForEach(
 					[&PolyHash, &RefVertexInstanceID](const FName AttributeName, auto AttributeArrayRef)
 					{
-						for (int32 Index = 0; Index < AttributeArrayRef.GetNumIndices(); ++Index)
+						for (int32 Index = 0; Index < AttributeArrayRef.GetNumChannels(); ++Index)
 						{
 							PolyHash = HashCombine(PolyHash, GetTypeHash(AttributeArrayRef.Get(RefVertexInstanceID)));
 						}
@@ -1762,7 +1895,7 @@ bool FStaticMeshOperations::GenerateUniqueUVsForStaticMesh(const FMeshDescriptio
 			}
 
 			FPolygonID& UniquePoly = UniquePolygons.FindOrAdd(PolyHash);
-			if (UniquePoly == FPolygonID::Invalid)
+			if (UniquePoly == INDEX_NONE)
 			{
 				UniquePoly = RefPolygonID;
 				for (FVertexInstanceID RefVertexInstanceID : RefVertexInstances)
@@ -1772,7 +1905,7 @@ bool FStaticMeshOperations::GenerateUniqueUVsForStaticMesh(const FMeshDescriptio
 			}
 			else
 			{
-				const TArray<FVertexInstanceID>& TestVertexInstances = DuplicateMeshDescription.GetPolygonVertexInstances(UniquePoly);
+				TArray<FVertexInstanceID, TInlineAllocator<4>> TestVertexInstances = DuplicateMeshDescription.GetPolygonVertexInstances<TInlineAllocator<4>>(UniquePoly);
 				int32 VertexInstanceIndex = 0;
 				for (FVertexInstanceID RefVertexInstanceID : RefVertexInstances)
 				{
@@ -1808,7 +1941,7 @@ bool FStaticMeshOperations::GenerateUniqueUVsForStaticMesh(const FMeshDescriptio
 		TVertexInstanceAttributesConstRef<FVector2D> DupVertexInstanceUVs = DuplicateMeshDescription.VertexInstanceAttributes().GetAttributesRef<FVector2D>(MeshAttribute::VertexInstance::TextureCoordinate);
 		TVertexInstanceAttributesConstRef<FVector2D> VertexInstanceUVs = MeshDescription.VertexInstanceAttributes().GetAttributesRef<FVector2D>(MeshAttribute::VertexInstance::TextureCoordinate);
 		// Save generated UVs
-		check(DupVertexInstanceUVs.GetNumIndices() > 1);
+		check(DupVertexInstanceUVs.GetNumChannels() > 1);
 		OutTexCoords.AddZeroed(VertexInstanceUVs.GetNumElements());
 		int32 TextureCoordIndex = 0;
 		for (const FVertexInstanceID& VertexInstanceID : MeshDescription.VertexInstances().GetElementIDs())
@@ -1826,51 +1959,51 @@ bool FStaticMeshOperations::GenerateUniqueUVsForStaticMesh(const FMeshDescriptio
 bool FStaticMeshOperations::AddUVChannel(FMeshDescription& MeshDescription)
 {
 	TVertexInstanceAttributesRef<FVector2D> VertexInstanceUVs = MeshDescription.VertexInstanceAttributes().GetAttributesRef<FVector2D>(MeshAttribute::VertexInstance::TextureCoordinate);
-	if (VertexInstanceUVs.GetNumIndices() >= MAX_MESH_TEXTURE_COORDS)
+	if (VertexInstanceUVs.GetNumChannels() >= MAX_MESH_TEXTURE_COORDS)
 	{
 		UE_LOG(LogStaticMeshOperations, Error, TEXT("AddUVChannel: Cannot add UV channel. Maximum number of UV channels reached (%d)."), MAX_MESH_TEXTURE_COORDS);
 		return false;
 	}
 
-	VertexInstanceUVs.SetNumIndices(VertexInstanceUVs.GetNumIndices() + 1);
+	VertexInstanceUVs.SetNumChannels(VertexInstanceUVs.GetNumChannels() + 1);
 	return true;
 }
 
 bool FStaticMeshOperations::InsertUVChannel(FMeshDescription& MeshDescription, int32 UVChannelIndex)
 {
 	TVertexInstanceAttributesRef<FVector2D> VertexInstanceUVs = MeshDescription.VertexInstanceAttributes().GetAttributesRef<FVector2D>(MeshAttribute::VertexInstance::TextureCoordinate);
-	if (UVChannelIndex < 0 || UVChannelIndex > VertexInstanceUVs.GetNumIndices())
+	if (UVChannelIndex < 0 || UVChannelIndex > VertexInstanceUVs.GetNumChannels())
 	{
 		UE_LOG(LogStaticMeshOperations, Error, TEXT("InsertUVChannel: Cannot insert UV channel. Given UV channel index %d is out of bounds."), UVChannelIndex);
 		return false;
 	}
 
-	if (VertexInstanceUVs.GetNumIndices() >= MAX_MESH_TEXTURE_COORDS)
+	if (VertexInstanceUVs.GetNumChannels() >= MAX_MESH_TEXTURE_COORDS)
 	{
 		UE_LOG(LogStaticMeshOperations, Error, TEXT("InsertUVChannel: Cannot insert UV channel. Maximum number of UV channels reached (%d)."), MAX_MESH_TEXTURE_COORDS);
 		return false;
 	}
 
-	VertexInstanceUVs.InsertIndex(UVChannelIndex);
+	VertexInstanceUVs.InsertChannel(UVChannelIndex);
 	return true;
 }
 
 bool FStaticMeshOperations::RemoveUVChannel(FMeshDescription& MeshDescription, int32 UVChannelIndex)
 {
 	TVertexInstanceAttributesRef<FVector2D> VertexInstanceUVs = MeshDescription.VertexInstanceAttributes().GetAttributesRef<FVector2D>(MeshAttribute::VertexInstance::TextureCoordinate);
-	if (VertexInstanceUVs.GetNumIndices() == 1)
+	if (VertexInstanceUVs.GetNumChannels() == 1)
 	{
 		UE_LOG(LogStaticMeshOperations, Error, TEXT("RemoveUVChannel: Cannot remove UV channel. There must be at least one channel."));
 		return false;
 	}
 
-	if (UVChannelIndex < 0 || UVChannelIndex >= VertexInstanceUVs.GetNumIndices())
+	if (UVChannelIndex < 0 || UVChannelIndex >= VertexInstanceUVs.GetNumChannels())
 	{
 		UE_LOG(LogStaticMeshOperations, Error, TEXT("RemoveUVChannel: Cannot remove UV channel. Given UV channel index %d is out of bounds."), UVChannelIndex);
 		return false;
 	}
 
-	VertexInstanceUVs.RemoveIndex(UVChannelIndex);
+	VertexInstanceUVs.RemoveChannel(UVChannelIndex);
 	return true;
 }
 
@@ -1942,7 +2075,7 @@ void FStaticMeshOperations::GenerateCylindricalUV(FMeshDescription& MeshDescript
 	// Fix the UV coordinates for triangles at the seam where the angle wraps around
 	for (const FPolygonID& PolygonID : MeshDescription.Polygons().GetElementIDs())
 	{
-		const TArray<FVertexInstanceID>& VertexInstances = MeshDescription.GetPolygonVertexInstances(PolygonID);
+		TArray<FVertexInstanceID, TInlineAllocator<4>> VertexInstances = MeshDescription.GetPolygonVertexInstances<TInlineAllocator<4>>(PolygonID);
 		int32 NumInstances = VertexInstances.Num();
 		if (NumInstances >= 2)
 		{
@@ -2011,14 +2144,14 @@ void FStaticMeshOperations::GenerateBoxUV(const FMeshDescription& MeshDescriptio
 	BoxPlanes.Add(FPlane(Center - FVector(HalfSize.X, 0, 0), FVector::BackwardVector));	// Back plane
 
 	// For each polygon, find the box plane that best matches the polygon normal
-	for (const FPolygonID& PolygonID : MeshDescription.Polygons().GetElementIDs())
+	for (const FTriangleID TriangleID : MeshDescription.Triangles().GetElementIDs())
 	{
-		const TArray<FVertexInstanceID>& VertexInstances = MeshDescription.GetPolygonVertexInstances(PolygonID);
-		check(VertexInstances.Num() == 3);
+		TArrayView<const FVertexID> Vertices = MeshDescription.GetTriangleVertices(TriangleID);
+		TArrayView<const FVertexInstanceID> VertexInstances = MeshDescription.GetTriangleVertexInstances(TriangleID);
 
-		FVector Vertex0 = VertexPositions[MeshDescription.GetVertexInstanceVertex(VertexInstances[0])];
-		FVector Vertex1 = VertexPositions[MeshDescription.GetVertexInstanceVertex(VertexInstances[1])];
-		FVector Vertex2 = VertexPositions[MeshDescription.GetVertexInstanceVertex(VertexInstances[2])];
+		FVector Vertex0 = VertexPositions[Vertices[0]];
+		FVector Vertex1 = VertexPositions[Vertices[1]];
+		FVector Vertex2 = VertexPositions[Vertices[2]];
 
 		FPlane PolygonPlane(Vertex0, Vertex2, Vertex1);
 
@@ -2040,7 +2173,7 @@ void FStaticMeshOperations::GenerateBoxUV(const FMeshDescription& MeshDescriptio
 		FVector V = PlaneUVs[BestPlaneIndex].Value;
 		FVector Offset = Params.Position - HalfSize * (U + V);
 
-		for (const FVertexInstanceID& VertexInstanceID : VertexInstances)
+		for (const FVertexInstanceID VertexInstanceID : VertexInstances)
 		{
 			const FVertexID VertexID = MeshDescription.GetVertexInstanceVertex(VertexInstanceID);
 			FVector Vertex = VertexPositions[VertexID];
@@ -2198,25 +2331,25 @@ FSHAHash FStaticMeshOperations::ComputeSHAHash(const FMeshDescription& MeshDescr
 void FStaticMeshOperations::FlipPolygons(FMeshDescription& MeshDescription)
 {
 	TSet<FVertexInstanceID> VertexInstanceIDs;
-	for (FPolygonID PolygonID : MeshDescription.Polygons().GetElementIDs())
+	for (FTriangleID TriangleID : MeshDescription.Triangles().GetElementIDs())
 	{
-		VertexInstanceIDs.Append(MeshDescription.GetPolygonVertexInstances(PolygonID));
-		MeshDescription.ReversePolygonFacing(PolygonID);
+		TArrayView<const FVertexInstanceID> TriVertInstances = MeshDescription.GetTriangleVertexInstances(TriangleID);
+		for (const FVertexInstanceID TriVertInstance : TriVertInstances)
+		{
+			VertexInstanceIDs.Add(TriVertInstance);
+		}
+		MeshDescription.ReverseTriangleFacing(TriangleID);
 	}
 
 	// Flip tangents and normals
-	const TVertexInstanceAttributesRef<FVector> VertexNormals = MeshDescription.VertexInstanceAttributes().GetAttributesRef<FVector>(MeshAttribute::VertexInstance::Normal);
-	const TVertexInstanceAttributesRef<FVector> VertexTangents = MeshDescription.VertexInstanceAttributes().GetAttributesRef<FVector>(MeshAttribute::VertexInstance::Tangent);
+	TVertexInstanceAttributesRef<FVector> VertexNormals = MeshDescription.VertexInstanceAttributes().GetAttributesRef<FVector>(MeshAttribute::VertexInstance::Normal);
+	TVertexInstanceAttributesRef<FVector> VertexTangents = MeshDescription.VertexInstanceAttributes().GetAttributesRef<FVector>(MeshAttribute::VertexInstance::Tangent);
 
 	for (const FVertexInstanceID VertexInstanceID : VertexInstanceIDs)
 	{
 		// Just reverse the sign of the normals/tangents; note that since binormals are the cross product of normal with tangent, they are left untouched
-		FVector Normal = VertexNormals[VertexInstanceID] * -1.0f;
-		FVector Tangent = VertexTangents[VertexInstanceID] * -1.0f;
-
-		TAttributesSet<FVertexInstanceID>& AttributesSet = MeshDescription.VertexInstanceAttributes();
-		AttributesSet.SetAttribute(VertexInstanceID, MeshAttribute::VertexInstance::Normal, 0, Normal);
-		AttributesSet.SetAttribute(VertexInstanceID, MeshAttribute::VertexInstance::Tangent, 0, Tangent);
+		VertexNormals[VertexInstanceID] *= -1.0f;
+		VertexTangents[VertexInstanceID] *= -1.0f;
 	}
 }
 
diff --git a/Engine/Source/Runtime/StaticMeshDescription/Public/StaticMeshAttributes.h b/Engine/Source/Runtime/StaticMeshDescription/Public/StaticMeshAttributes.h
--- a/Engine/Source/Runtime/StaticMeshDescription/Public/StaticMeshAttributes.h
+++ b/Engine/Source/Runtime/StaticMeshDescription/Public/StaticMeshAttributes.h
@@ -9,7 +9,8 @@ namespace MeshAttribute
 {
 	namespace Vertex
 	{
-		extern STATICMESHDESCRIPTION_API const FName CornerSharpness;
+		UE_DEPRECATED(4.26, "This attribute is no longer supported, please remove code pertaining to it.")
+		extern STATICMESHDESCRIPTION_API const FName CornerSharpness;	// deprecated
 	}
 
 	namespace VertexInstance
@@ -19,32 +20,50 @@ namespace MeshAttribute
 		extern STATICMESHDESCRIPTION_API const FName Tangent;
 		extern STATICMESHDESCRIPTION_API const FName BinormalSign;
 		extern STATICMESHDESCRIPTION_API const FName Color;
+		extern STATICMESHDESCRIPTION_API const FName UVIndex;
 	}
 
 	namespace Edge
 	{
 		extern STATICMESHDESCRIPTION_API const FName IsHard;
-		extern STATICMESHDESCRIPTION_API const FName IsUVSeam;
-		extern STATICMESHDESCRIPTION_API const FName CreaseSharpness;
+		UE_DEPRECATED(4.26, "This attribute is no longer supported, please remove code pertaining to it.")
+		extern STATICMESHDESCRIPTION_API const FName IsUVSeam;			// deprecated
+		UE_DEPRECATED(4.26, "This attribute is no longer supported, please remove code pertaining to it.")
+		extern STATICMESHDESCRIPTION_API const FName CreaseSharpness;	// deprecated
+	}
+
+	namespace Triangle
+	{
+		extern STATICMESHDESCRIPTION_API const FName Normal;
+		extern STATICMESHDESCRIPTION_API const FName Tangent;
+		extern STATICMESHDESCRIPTION_API const FName Binormal;
 	}
 
 	namespace Polygon
 	{
+		UE_DEPRECATED(4.26, "This attribute is no longer supported, please remove code pertaining to it.")
 		extern STATICMESHDESCRIPTION_API const FName Normal;
+		UE_DEPRECATED(4.26, "This attribute is no longer supported, please remove code pertaining to it.")
 		extern STATICMESHDESCRIPTION_API const FName Tangent;
+		UE_DEPRECATED(4.26, "This attribute is no longer supported, please remove code pertaining to it.")
 		extern STATICMESHDESCRIPTION_API const FName Binormal;
+		UE_DEPRECATED(4.26, "This attribute is no longer supported, please remove code pertaining to it.")
 		extern STATICMESHDESCRIPTION_API const FName Center;
 	}
 
 	namespace PolygonGroup
 	{
 		extern STATICMESHDESCRIPTION_API const FName ImportedMaterialSlotName;
-		extern STATICMESHDESCRIPTION_API const FName EnableCollision;
-		extern STATICMESHDESCRIPTION_API const FName CastShadow;
+		UE_DEPRECATED(4.26, "This attribute is no longer supported, please remove code pertaining to it.")
+		extern STATICMESHDESCRIPTION_API const FName EnableCollision;	// deprecated
+		UE_DEPRECATED(4.26, "This attribute is no longer supported, please remove code pertaining to it.")
+		extern STATICMESHDESCRIPTION_API const FName CastShadow;		// deprecated
 	}
 }
 
 
+PRAGMA_DISABLE_DEPRECATION_WARNINGS
+
 class STATICMESHDESCRIPTION_API FStaticMeshAttributes : public FMeshAttributes
 {
 public:
@@ -55,9 +74,14 @@ public:
 
 	virtual void Register() override;
 
+	UE_DEPRECATED(4.26, "Please use RegisterTriangleNormalAndTangentAttributes() instead.")
 	void RegisterPolygonNormalAndTangentAttributes();
 
+	void RegisterTriangleNormalAndTangentAttributes();
+
+	UE_DEPRECATED(4.26, "This attribute is no longer supported, please remove code pertaining to it.")
 	TVertexAttributesRef<float> GetVertexCornerSharpnesses() { return MeshDescription.VertexAttributes().GetAttributesRef<float>(MeshAttribute::Vertex::CornerSharpness); }
+	UE_DEPRECATED(4.26, "This attribute is no longer supported, please remove code pertaining to it.")
 	TVertexAttributesConstRef<float> GetVertexCornerSharpnesses() const { return MeshDescription.VertexAttributes().GetAttributesRef<float>(MeshAttribute::Vertex::CornerSharpness); }
 
 	TVertexInstanceAttributesRef<FVector2D> GetVertexInstanceUVs() { return MeshDescription.VertexInstanceAttributes().GetAttributesRef<FVector2D>(MeshAttribute::VertexInstance::TextureCoordinate); }
@@ -78,19 +102,38 @@ public:
 	TEdgeAttributesRef<bool> GetEdgeHardnesses() { return MeshDescription.EdgeAttributes().GetAttributesRef<bool>(MeshAttribute::Edge::IsHard); }
 	TEdgeAttributesConstRef<bool> GetEdgeHardnesses() const { return MeshDescription.EdgeAttributes().GetAttributesRef<bool>(MeshAttribute::Edge::IsHard); }
 
+	UE_DEPRECATED(4.26, "This attribute is no longer supported, please remove code pertaining to it.")
 	TEdgeAttributesRef<float> GetEdgeCreaseSharpnesses() { return MeshDescription.EdgeAttributes().GetAttributesRef<float>(MeshAttribute::Edge::CreaseSharpness); }
+	UE_DEPRECATED(4.26, "This attribute is no longer supported, please remove code pertaining to it.")
 	TEdgeAttributesConstRef<float> GetEdgeCreaseSharpnesses() const { return MeshDescription.EdgeAttributes().GetAttributesRef<float>(MeshAttribute::Edge::CreaseSharpness); }
 
+	TTriangleAttributesRef<FVector> GetTriangleNormals() { return MeshDescription.TriangleAttributes().GetAttributesRef<FVector>(MeshAttribute::Triangle::Normal); }
+	TTriangleAttributesConstRef<FVector> GetTriangleNormals() const { return MeshDescription.TriangleAttributes().GetAttributesRef<FVector>(MeshAttribute::Triangle::Normal); }
+
+	TTriangleAttributesRef<FVector> GetTriangleTangents() { return MeshDescription.TriangleAttributes().GetAttributesRef<FVector>(MeshAttribute::Triangle::Tangent); }
+	TTriangleAttributesConstRef<FVector> GetTriangleTangents() const { return MeshDescription.TriangleAttributes().GetAttributesRef<FVector>(MeshAttribute::Triangle::Tangent); }
+
+	TTriangleAttributesRef<FVector> GetTriangleBinormals() { return MeshDescription.TriangleAttributes().GetAttributesRef<FVector>(MeshAttribute::Triangle::Binormal); }
+	TTriangleAttributesConstRef<FVector> GetTriangleBinormals() const { return MeshDescription.TriangleAttributes().GetAttributesRef<FVector>(MeshAttribute::Triangle::Binormal); }
+
+	UE_DEPRECATED(4.26, "This attribute is no longer supported, please remove code pertaining to it.")
 	TPolygonAttributesRef<FVector> GetPolygonNormals() { return MeshDescription.PolygonAttributes().GetAttributesRef<FVector>(MeshAttribute::Polygon::Normal); }
+	UE_DEPRECATED(4.26, "This attribute is no longer supported, please remove code pertaining to it.")
 	TPolygonAttributesConstRef<FVector> GetPolygonNormals() const { return MeshDescription.PolygonAttributes().GetAttributesRef<FVector>(MeshAttribute::Polygon::Normal); }
 
+	UE_DEPRECATED(4.26, "This attribute is no longer supported, please remove code pertaining to it.")
 	TPolygonAttributesRef<FVector> GetPolygonTangents() { return MeshDescription.PolygonAttributes().GetAttributesRef<FVector>(MeshAttribute::Polygon::Tangent); }
+	UE_DEPRECATED(4.26, "This attribute is no longer supported, please remove code pertaining to it.")
 	TPolygonAttributesConstRef<FVector> GetPolygonTangents() const { return MeshDescription.PolygonAttributes().GetAttributesRef<FVector>(MeshAttribute::Polygon::Tangent); }
 
+	UE_DEPRECATED(4.26, "This attribute is no longer supported, please remove code pertaining to it.")
 	TPolygonAttributesRef<FVector> GetPolygonBinormals() { return MeshDescription.PolygonAttributes().GetAttributesRef<FVector>(MeshAttribute::Polygon::Binormal); }
+	UE_DEPRECATED(4.26, "This attribute is no longer supported, please remove code pertaining to it.")
 	TPolygonAttributesConstRef<FVector> GetPolygonBinormals() const { return MeshDescription.PolygonAttributes().GetAttributesRef<FVector>(MeshAttribute::Polygon::Binormal); }
 
+	UE_DEPRECATED(4.26, "This attribute is no longer supported, please remove code pertaining to it.")
 	TPolygonAttributesRef<FVector> GetPolygonCenters() { return MeshDescription.PolygonAttributes().GetAttributesRef<FVector>(MeshAttribute::Polygon::Center); }
+	UE_DEPRECATED(4.26, "This attribute is no longer supported, please remove code pertaining to it.")
 	TPolygonAttributesConstRef<FVector> GetPolygonCenters() const { return MeshDescription.PolygonAttributes().GetAttributesRef<FVector>(MeshAttribute::Polygon::Center); }
 
 	TPolygonGroupAttributesRef<FName> GetPolygonGroupMaterialSlotNames() { return MeshDescription.PolygonGroupAttributes().GetAttributesRef<FName>(MeshAttribute::PolygonGroup::ImportedMaterialSlotName); }
@@ -106,17 +149,29 @@ public:
 		: FMeshConstAttributes(InMeshDescription)
 	{}
 
+	UE_DEPRECATED(4.26, "This attribute is no longer supported, please remove code pertaining to it.")
 	TVertexAttributesConstRef<float> GetVertexCornerSharpnesses() const { return MeshDescription.VertexAttributes().GetAttributesRef<float>(MeshAttribute::Vertex::CornerSharpness); }
 	TVertexInstanceAttributesConstRef<FVector2D> GetVertexInstanceUVs() const { return MeshDescription.VertexInstanceAttributes().GetAttributesRef<FVector2D>(MeshAttribute::VertexInstance::TextureCoordinate); }
 	TVertexInstanceAttributesConstRef<FVector> GetVertexInstanceNormals() const { return MeshDescription.VertexInstanceAttributes().GetAttributesRef<FVector>(MeshAttribute::VertexInstance::Normal); }
 	TVertexInstanceAttributesConstRef<FVector> GetVertexInstanceTangents() const { return MeshDescription.VertexInstanceAttributes().GetAttributesRef<FVector>(MeshAttribute::VertexInstance::Tangent); }
 	TVertexInstanceAttributesConstRef<float> GetVertexInstanceBinormalSigns() const { return MeshDescription.VertexInstanceAttributes().GetAttributesRef<float>(MeshAttribute::VertexInstance::BinormalSign); }
 	TVertexInstanceAttributesConstRef<FVector4> GetVertexInstanceColors() const { return MeshDescription.VertexInstanceAttributes().GetAttributesRef<FVector4>(MeshAttribute::VertexInstance::Color); }
+	TVertexInstanceAttributesConstRef<int> GetVertexInstanceUVIndices() const { return MeshDescription.VertexInstanceAttributes().GetAttributesRef<int>(MeshAttribute::VertexInstance::UVIndex); }
 	TEdgeAttributesConstRef<bool> GetEdgeHardnesses() const { return MeshDescription.EdgeAttributes().GetAttributesRef<bool>(MeshAttribute::Edge::IsHard); }
+	UE_DEPRECATED(4.26, "This attribute is no longer supported, please remove code pertaining to it.")
 	TEdgeAttributesConstRef<float> GetEdgeCreaseSharpnesses() const { return MeshDescription.EdgeAttributes().GetAttributesRef<float>(MeshAttribute::Edge::CreaseSharpness); }
+	TTriangleAttributesConstRef<FVector> GetTriangleNormals() const { return MeshDescription.TriangleAttributes().GetAttributesRef<FVector>(MeshAttribute::Triangle::Normal); }
+	TTriangleAttributesConstRef<FVector> GetTriangleTangents() const { return MeshDescription.TriangleAttributes().GetAttributesRef<FVector>(MeshAttribute::Triangle::Tangent); }
+	TTriangleAttributesConstRef<FVector> GetTriangleBinormals() const { return MeshDescription.TriangleAttributes().GetAttributesRef<FVector>(MeshAttribute::Triangle::Binormal); }
+	UE_DEPRECATED(4.26, "This attribute is no longer supported, please remove code pertaining to it.")
 	TPolygonAttributesConstRef<FVector> GetPolygonNormals() const { return MeshDescription.PolygonAttributes().GetAttributesRef<FVector>(MeshAttribute::Polygon::Normal); }
+	UE_DEPRECATED(4.26, "This attribute is no longer supported, please remove code pertaining to it.")
 	TPolygonAttributesConstRef<FVector> GetPolygonTangents() const { return MeshDescription.PolygonAttributes().GetAttributesRef<FVector>(MeshAttribute::Polygon::Tangent); }
+	UE_DEPRECATED(4.26, "This attribute is no longer supported, please remove code pertaining to it.")
 	TPolygonAttributesConstRef<FVector> GetPolygonBinormals() const { return MeshDescription.PolygonAttributes().GetAttributesRef<FVector>(MeshAttribute::Polygon::Binormal); }
+	UE_DEPRECATED(4.26, "This attribute is no longer supported, please remove code pertaining to it.")
 	TPolygonAttributesConstRef<FVector> GetPolygonCenters() const { return MeshDescription.PolygonAttributes().GetAttributesRef<FVector>(MeshAttribute::Polygon::Center); }
 	TPolygonGroupAttributesConstRef<FName> GetPolygonGroupMaterialSlotNames() const { return MeshDescription.PolygonGroupAttributes().GetAttributesRef<FName>(MeshAttribute::PolygonGroup::ImportedMaterialSlotName); }
 };
+
+PRAGMA_ENABLE_DEPRECATION_WARNINGS
diff --git a/Engine/Source/Runtime/StaticMeshDescription/Public/StaticMeshDescription.h b/Engine/Source/Runtime/StaticMeshDescription/Public/StaticMeshDescription.h
--- a/Engine/Source/Runtime/StaticMeshDescription/Public/StaticMeshDescription.h
+++ b/Engine/Source/Runtime/StaticMeshDescription/Public/StaticMeshDescription.h
@@ -9,6 +9,8 @@
 
 class UMaterial;
 
+PRAGMA_DISABLE_DEPRECATION_WARNINGS
+
 /**
 * A wrapper for MeshDescription, customized for static meshes
 */
@@ -51,7 +53,9 @@ public:
 
 public:
 
+	UE_DEPRECATED(4.25, "This attribute is no longer supported, please remove code pertaining to it.")
 	TVertexAttributesRef<float> GetVertexCornerSharpnesses() { return GetRequiredAttributes().GetVertexCornerSharpnesses(); }
+	UE_DEPRECATED(4.25, "This attribute is no longer supported, please remove code pertaining to it.")
 	TVertexAttributesConstRef<float> GetVertexCornerSharpnesses() const { return GetRequiredAttributes().GetVertexCornerSharpnesses(); }
 
 	TVertexInstanceAttributesRef<FVector2D> GetVertexInstanceUVs() { return GetRequiredAttributes().GetVertexInstanceUVs(); }
@@ -72,9 +76,13 @@ public:
 	TEdgeAttributesRef<bool> GetEdgeHardnesses() { return GetRequiredAttributes().GetEdgeHardnesses(); }
 	TEdgeAttributesConstRef<bool> GetEdgeHardnesses() const { return GetRequiredAttributes().GetEdgeHardnesses(); }
 
+	UE_DEPRECATED(4.25, "This attribute is no longer supported, please remove code pertaining to it.")
 	TEdgeAttributesRef<float> GetEdgeCreaseSharpnesses() { return GetRequiredAttributes().GetEdgeCreaseSharpnesses(); }
+	UE_DEPRECATED(4.25, "This attribute is no longer supported, please remove code pertaining to it.")
 	TEdgeAttributesConstRef<float> GetEdgeCreaseSharpnesses() const { return GetRequiredAttributes().GetEdgeCreaseSharpnesses(); }
 
 	TPolygonGroupAttributesRef<FName> GetPolygonGroupMaterialSlotNames() { return GetRequiredAttributes().GetPolygonGroupMaterialSlotNames(); }
 	TPolygonGroupAttributesConstRef<FName> GetPolygonGroupMaterialSlotNames() const { return GetRequiredAttributes().GetPolygonGroupMaterialSlotNames(); }
 };
+
+PRAGMA_ENABLE_DEPRECATION_WARNINGS
diff --git a/Engine/Source/Runtime/StaticMeshDescription/Public/StaticMeshOperations.h b/Engine/Source/Runtime/StaticMeshDescription/Public/StaticMeshOperations.h
--- a/Engine/Source/Runtime/StaticMeshDescription/Public/StaticMeshOperations.h
+++ b/Engine/Source/Runtime/StaticMeshDescription/Public/StaticMeshOperations.h
@@ -54,8 +54,12 @@ public:
 	};
 
 	/** Set the polygon tangent, normal, binormal and polygonCenter for all polygons in the mesh description. */
+	UE_DEPRECATED(4.26, "Please use ComputeTriangleTangentsAndNormals() instead.")
 	static void ComputePolygonTangentsAndNormals(FMeshDescription& MeshDescription, float ComparisonThreshold = 0.0f);
 
+	/** Set the triangle tangent, normal, binormal and triangleCenter for all triangles in the mesh description. */
+	static void ComputeTriangleTangentsAndNormals(FMeshDescription& MeshDescription, float ComparisonThreshold = 0.0f);
+
 	/** 
 	 * Recompute any invalid normal, tangent or Bi-Normal for every vertex in the mesh description with the given options.
 	 * If the EComputeNTBsFlags Normals or Tangents are set, the corresponding data will be force-recomputed.
