Moved/added some deprecation pragmas to fix warnings on some compilers (e.g. VS2017).

diff --git a/Engine/Source/Editor/MainFrame/Private/MainFrameModule.h b/Engine/Source/Editor/MainFrame/Private/MainFrameModule.h
--- a/Engine/Source/Editor/MainFrame/Private/MainFrameModule.h
+++ b/Engine/Source/Editor/MainFrame/Private/MainFrameModule.h
@@ -15,6 +15,7 @@
 /**
  * Editor main frame module
  */
+PRAGMA_DISABLE_DEPRECATION_WARNINGS
 class FMainFrameModule
 	: public IMainFrameModule
 {
@@ -26,10 +27,9 @@ public:
 	virtual void RecreateDefaultMainFrame(const bool bStartImmersive, const bool bStartPIE) override;
 	virtual bool IsRecreatingDefaultMainFrame() const override;
 	virtual TSharedRef<SWidget> MakeMainMenu(const TSharedPtr<FTabManager>& TabManager, const FName MenuName, FToolMenuContext& ToolMenuContext) const override;
-
-PRAGMA_DISABLE_DEPRECATION_WARNINGS
+	
+	// deprecated in 4.26
 	virtual TSharedRef<SWidget> MakeMainTabMenu(const TSharedPtr<FTabManager>& TabManager, const FName MenuName, FToolMenuContext& ToolMenuContext) const override;
-PRAGMA_ENABLE_DEPRECATION_WARNINGS
 
 	virtual TSharedRef<SWidget> MakeDeveloperTools( const TArray<FMainFrameDeveloperTool>& AdditionalTools ) const override;
 
@@ -258,3 +258,5 @@ private:
 	// Is recreating Default Main Frame
 	bool bRecreatingDefaultMainFrame;
 };
+PRAGMA_ENABLE_DEPRECATION_WARNINGS
+
diff --git a/Engine/Source/Editor/UnrealEd/Classes/Editor/EditorEngine.h b/Engine/Source/Editor/UnrealEd/Classes/Editor/EditorEngine.h
--- a/Engine/Source/Editor/UnrealEd/Classes/Editor/EditorEngine.h
+++ b/Engine/Source/Editor/UnrealEd/Classes/Editor/EditorEngine.h
@@ -343,6 +343,8 @@ public:
  * Engine that drives the Editor.
  * Separate from UGameEngine because it may have much different functionality than desired for an instance of a game itself.
  */
+PRAGMA_DISABLE_DEPRECATION_WARNINGS
+
 UCLASS(config=Engine, transient)
 class UNREALED_API UEditorEngine : public UEngine
 {
@@ -836,9 +838,8 @@ private:
 	virtual TSharedPtr<SViewport> GetGameViewportWidget() const override;
 	virtual void TriggerStreamingDataRebuild() override;
 
-	PRAGMA_DISABLE_DEPRECATION_WARNINGS
+	// deprecated in 4.26
 	virtual bool NetworkRemapPath(UNetDriver* Driver, FString& Str, bool bReading = true) override;
-	PRAGMA_ENABLE_DEPRECATION_WARNINGS
 	virtual bool NetworkRemapPath(UNetConnection* Connection, FString& Str, bool bReading = true) override;
 	virtual bool NetworkRemapPath(UPendingNetGame* PendingNetGame, FString& Str, bool bReading = true) override;
 
@@ -3380,6 +3381,7 @@ protected:
 	/** Number of currently running instances logged into an online platform */
 	int32 NumOnlinePIEInstances;
 };
+PRAGMA_ENABLE_DEPRECATION_WARNINGS
 
 //////////////////////////////////////////////////////////////////////////
 // FActorLabelUtilities
diff --git a/Engine/Source/Runtime/StaticMeshDescription/Private/StaticMeshAttributes.cpp b/Engine/Source/Runtime/StaticMeshDescription/Private/StaticMeshAttributes.cpp
--- a/Engine/Source/Runtime/StaticMeshDescription/Private/StaticMeshAttributes.cpp
+++ b/Engine/Source/Runtime/StaticMeshDescription/Private/StaticMeshAttributes.cpp
@@ -4,8 +4,6 @@
 
 namespace MeshAttribute
 {
-	const FName Vertex::CornerSharpness("CornerSharpness");
-
 	const FName VertexInstance::TextureCoordinate("TextureCoordinate");
 	const FName VertexInstance::Normal("Normal");
 	const FName VertexInstance::Tangent("Tangent");
@@ -13,8 +11,6 @@ namespace MeshAttribute
 	const FName VertexInstance::Color("Color");
 
 	const FName Edge::IsHard("IsHard");
-	const FName Edge::IsUVSeam("IsUVSeam");
-	const FName Edge::CreaseSharpness("CreaseSharpness");
 
 	const FName Triangle::Normal("Normal");
 	const FName Triangle::Tangent("Tangent");
@@ -26,8 +22,15 @@ namespace MeshAttribute
 	const FName Polygon::Center("Center");
 
 	const FName PolygonGroup::ImportedMaterialSlotName("ImportedMaterialSlotName");
-	const FName PolygonGroup::EnableCollision("EnableCollision");
+
+PRAGMA_DISABLE_DEPRECATION_WARNINGS
+	// all deprected in 4.26
+	const FName Vertex::CornerSharpness("CornerSharpness");
 	const FName PolygonGroup::CastShadow("CastShadow");
+	const FName Edge::CreaseSharpness("CreaseSharpness");
+	const FName Edge::IsUVSeam("IsUVSeam");
+	const FName PolygonGroup::EnableCollision("EnableCollision");
+PRAGMA_ENABLE_DEPRECATION_WARNINGS
 }
 
 
