Fix mikkt crashing when meshdescription has no triangles

diff --git a/Engine/Source/Runtime/StaticMeshDescription/Private/StaticMeshOperations.cpp b/Engine/Source/Runtime/StaticMeshDescription/Private/StaticMeshOperations.cpp
--- a/Engine/Source/Runtime/StaticMeshDescription/Private/StaticMeshOperations.cpp
+++ b/Engine/Source/Runtime/StaticMeshDescription/Private/StaticMeshOperations.cpp
@@ -1637,6 +1637,12 @@ void FStaticMeshOperations::ComputeMikktTangents(FMeshDescription& MeshDescripti
 		MeshDescription.Compact(Remappings);
 	}
 
+	int32 NumTriangles = MeshDescription.Triangles().Num();
+	if (NumTriangles == 0)
+	{
+		return; // nothing to compute
+	}
+
 	// we can use mikktspace to calculate the tangents
 	SMikkTSpaceInterface MikkTInterface;
 	MikkTInterface.m_getNormal = MeshDescriptionMikktSpaceInterface::MikkGetNormal;
