FStaticMeshOperations::GenerateUniqueUVsForStaticMesh() shouldn't test all attributes when trying to merge UVs
* We only care for UVs & vertex colors

diff --git a/Engine/Source/Runtime/StaticMeshDescription/Private/StaticMeshOperations.cpp b/Engine/Source/Runtime/StaticMeshDescription/Private/StaticMeshOperations.cpp
--- a/Engine/Source/Runtime/StaticMeshDescription/Private/StaticMeshOperations.cpp
+++ b/Engine/Source/Runtime/StaticMeshDescription/Private/StaticMeshOperations.cpp
@@ -1879,6 +1879,9 @@ bool FStaticMeshOperations::GenerateUniqueUVsForStaticMesh(const FMeshDescriptio
 	//Remove the identical material
 	if (bMergeIdenticalMaterials)
 	{
+		TVertexInstanceAttributesConstRef<FVector2D> VertexInstanceUVs = DuplicateMeshDescription.VertexInstanceAttributes().GetAttributesRef<FVector2D>(MeshAttribute::VertexInstance::TextureCoordinate);
+		TVertexInstanceAttributesConstRef<FVector4> VertexColors = DuplicateMeshDescription.VertexInstanceAttributes().GetAttributesRef<FVector4>(MeshAttribute::VertexInstance::Color);
+
 		TArray<FPolygonID> ToDeletePolygons;
 		RemapVertexInstance.Reserve(DuplicateMeshDescription.VertexInstances().Num());
 
@@ -1888,6 +1891,17 @@ bool FStaticMeshOperations::GenerateUniqueUVsForStaticMesh(const FMeshDescriptio
 		UniquePolygons.Reserve(NumPolygons);
 		ToDeletePolygons.Reserve(NumPolygons);
 
+		auto HashAttribute = [](FVertexInstanceID InVertexInstanceID, auto InAttributeArrayRef, int32& OutPolyHash)
+		{
+			for (int32 Channel = 0; Channel < InAttributeArrayRef.GetNumChannels(); ++Channel)
+			{
+				for (const auto& Element : InAttributeArrayRef.GetArrayView(InVertexInstanceID, Channel))
+				{
+					OutPolyHash = HashCombine(OutPolyHash, GetTypeHash(Element));
+				}
+			}
+		};
+
 		for (FPolygonID RefPolygonID : DuplicateMeshDescription.Polygons().GetElementIDs())
 		{
 			const FPolygonGroupID RefPolygonGroupID = DuplicateMeshDescription.GetPolygonPolygonGroup(RefPolygonID);
@@ -1896,19 +1910,9 @@ bool FStaticMeshOperations::GenerateUniqueUVsForStaticMesh(const FMeshDescriptio
 			int32 PolyHash = GetTypeHash(RefPolygonGroupID);
 			for (FVertexInstanceID RefVertexInstanceID : RefVertexInstances)
 			{
-				// Compute hash based on all vertices attributes
-				DuplicateMeshDescription.VertexInstanceAttributes().ForEach(
-					[&PolyHash, &RefVertexInstanceID](const FName AttributeName, auto AttributeArrayRef)
-					{
-						for (int32 Channel = 0; Channel < AttributeArrayRef.GetNumChannels(); ++Channel)
-						{
-							for (const auto& Element : AttributeArrayRef.GetArrayView(RefVertexInstanceID, Channel))
-							{
-								PolyHash = HashCombine(PolyHash, GetTypeHash(Element));
-							}
-						}
-					}
-				);
+				// Compute hash based on UVs & vertices colors
+				HashAttribute(RefVertexInstanceID, VertexInstanceUVs, PolyHash);
+				HashAttribute(RefVertexInstanceID, VertexColors, PolyHash);
 			}
 
 			FPolygonID& UniquePoly = UniquePolygons.FindOrAdd(PolyHash);
