FStaticMeshOperations - Added AppendMeshDescriptions()
* A more efficient version of AppendMeshDescription() for multiple meshes
* Properly reserve arrays to limit amount of temporary allocations
* Add triangles rather than polygons, avoid the costlier SourceMesh.GetPolygonVertexInstances() which had to perform a map lookup

diff --git a/Engine/Source/Runtime/StaticMeshDescription/Private/StaticMeshOperations.cpp b/Engine/Source/Runtime/StaticMeshDescription/Private/StaticMeshOperations.cpp
--- a/Engine/Source/Runtime/StaticMeshDescription/Private/StaticMeshOperations.cpp
+++ b/Engine/Source/Runtime/StaticMeshDescription/Private/StaticMeshOperations.cpp
@@ -944,19 +944,9 @@ void FStaticMeshOperations::ConvertFromRawMesh(const FRawMesh& SourceRawMesh, FM
 	DestinationMeshDescription.BuildIndexers();
 }
 
-void FStaticMeshOperations::AppendMeshDescription(const FMeshDescription& SourceMesh, FMeshDescription& TargetMesh, FAppendSettings& AppendSettings)
+void FStaticMeshOperations::AppendMeshDescriptions(const TArray<const FMeshDescription*>& SourceMeshes, FMeshDescription& TargetMesh, FAppendSettings& AppendSettings)
 {
-	TRACE_CPUPROFILER_EVENT_SCOPE(FStaticMeshOperations::AppendMeshDescription);
-
-	FStaticMeshConstAttributes SourceAttributes(SourceMesh);
-	TVertexAttributesConstRef<FVector> SourceVertexPositions = SourceAttributes.GetVertexPositions();
-	TEdgeAttributesConstRef<bool> SourceEdgeHardnesses = SourceAttributes.GetEdgeHardnesses();
-	TPolygonGroupAttributesConstRef<FName> SourceImportedMaterialSlotNames = SourceAttributes.GetPolygonGroupMaterialSlotNames();
-	TVertexInstanceAttributesConstRef<FVector> SourceVertexInstanceNormals = SourceAttributes.GetVertexInstanceNormals();
-	TVertexInstanceAttributesConstRef<FVector> SourceVertexInstanceTangents = SourceAttributes.GetVertexInstanceTangents();
-	TVertexInstanceAttributesConstRef<float> SourceVertexInstanceBinormalSigns = SourceAttributes.GetVertexInstanceBinormalSigns();
-	TVertexInstanceAttributesConstRef<FVector4> SourceVertexInstanceColors = SourceAttributes.GetVertexInstanceColors();
-	TVertexInstanceAttributesConstRef<FVector2D> SourceVertexInstanceUVs = SourceAttributes.GetVertexInstanceUVs();
+	TRACE_CPUPROFILER_EVENT_SCOPE_STR((SourceMeshes.Num() > 1 ? "FStaticMeshOperations::AppendMeshDescriptions" : "FStaticMeshOperations::AppendMeshDescription"));
 
 	FStaticMeshAttributes TargetAttributes(TargetMesh);
 	TVertexAttributesRef<FVector> TargetVertexPositions = TargetAttributes.GetVertexPositions();
@@ -968,151 +958,213 @@ void FStaticMeshOperations::AppendMeshDescription(const FMeshDescription& Source
 	TVertexInstanceAttributesRef<FVector4> TargetVertexInstanceColors = TargetAttributes.GetVertexInstanceColors();
 	TVertexInstanceAttributesRef<FVector2D> TargetVertexInstanceUVs = TargetAttributes.GetVertexInstanceUVs();
 
-	//Copy into the target mesh
-	TargetMesh.ReserveNewVertices(SourceMesh.Vertices().Num());
-	TargetMesh.ReserveNewVertexInstances(SourceMesh.VertexInstances().Num());
-	TargetMesh.ReserveNewEdges(SourceMesh.Edges().Num());
-	TargetMesh.ReserveNewPolygons(SourceMesh.Polygons().Num());
+	TargetMesh.SuspendVertexInstanceIndexing();
+	TargetMesh.SuspendEdgeIndexing();
+	TargetMesh.SuspendPolygonIndexing();
+	TargetMesh.SuspendPolygonGroupIndexing();
+	TargetMesh.SuspendUVIndexing();
+
+	int32 NumVertices = 0;
+	int32 NumVertexInstances = 0;
+	int32 NumEdges = 0;
+	int32 NumTriangles = 0;
 
-	int32 NumSourceUVChannels = 0;
-	for (int32 ChannelIdx = 0; ChannelIdx < SourceVertexInstanceUVs.GetNumChannels(); ++ChannelIdx)
+	int32 MaxNumUVChannels = TargetVertexInstanceUVs.GetNumChannels();
+	int32 MaxNumPolygonGroups = 0;
+	int32 MaxNumMeshVertices = 0;
+	int32 MaxNumEdges = 0;
+	int32 MaxNumVertexInstances = 0;
+
+	for (const FMeshDescription* SourceMeshPtr : SourceMeshes)
 	{
-		if (AppendSettings.bMergeUVChannels[ChannelIdx])
+		const FMeshDescription& SourceMesh = *SourceMeshPtr;
+
+		NumVertices += SourceMesh.Vertices().Num();
+		NumVertexInstances += SourceMesh.VertexInstances().Num();
+		NumEdges += SourceMesh.Edges().Num();
+		NumTriangles += SourceMesh.Triangles().Num();
+
+		FStaticMeshConstAttributes SourceAttributes(SourceMesh);
+		TVertexInstanceAttributesConstRef<FVector2D> SourceVertexInstanceUVs = SourceAttributes.GetVertexInstanceUVs();
+		for (int32 ChannelIdx = MaxNumUVChannels; ChannelIdx < SourceVertexInstanceUVs.GetNumChannels(); ++ChannelIdx)
 		{
-			NumSourceUVChannels = ChannelIdx + 1;
+			if (AppendSettings.bMergeUVChannels[ChannelIdx])
+			{
+				MaxNumUVChannels = ChannelIdx + 1;
+			}
 		}
+
+		MaxNumPolygonGroups = FMath::Max(MaxNumPolygonGroups, SourceMesh.PolygonGroups().Num());
+		MaxNumMeshVertices = FMath::Max(MaxNumMeshVertices, SourceMesh.Vertices().Num());
+		MaxNumEdges = FMath::Max(MaxNumEdges, SourceMesh.Edges().Num());
+		MaxNumVertexInstances = FMath::Max(MaxNumVertexInstances, SourceMesh.VertexInstances().Num());
 	}
 
-	if (NumSourceUVChannels > TargetVertexInstanceUVs.GetNumChannels())
+	//Copy into the target mesh
+	TargetMesh.ReserveNewVertices(NumVertices);
+	TargetMesh.ReserveNewVertexInstances(NumVertexInstances);
+	TargetMesh.ReserveNewEdges(NumEdges);
+	TargetMesh.ReserveNewTriangles(NumTriangles);
+
+	if (MaxNumUVChannels > TargetVertexInstanceUVs.GetNumChannels())
 	{
-		TargetVertexInstanceUVs.SetNumChannels(NumSourceUVChannels);
+		TargetVertexInstanceUVs.SetNumChannels(MaxNumUVChannels);
 	}
 
-	//PolygonGroups
 	PolygonGroupMap RemapPolygonGroup;
-	if (AppendSettings.PolygonGroupsDelegate.IsBound())
-	{
-		AppendSettings.PolygonGroupsDelegate.Execute(SourceMesh, TargetMesh, RemapPolygonGroup);
-	}
-	else
+
+	TMap<FVertexID, FVertexID> SourceToTargetVertexID;
+	SourceToTargetVertexID.Reserve(MaxNumMeshVertices);
+	TMap<FVertexInstanceID, FVertexInstanceID> SourceToTargetVertexInstanceID;
+	SourceToTargetVertexInstanceID.Reserve(MaxNumVertexInstances);
+
+	for (const FMeshDescription* SourceMeshPtr : SourceMeshes)
 	{
-		for (FPolygonGroupID SourcePolygonGroupID : SourceMesh.PolygonGroups().GetElementIDs())
+		const FMeshDescription& SourceMesh = *SourceMeshPtr;
+
+		RemapPolygonGroup.Empty(MaxNumPolygonGroups);
+
+		FStaticMeshConstAttributes SourceAttributes(SourceMesh);
+		TVertexAttributesConstRef<FVector> SourceVertexPositions = SourceAttributes.GetVertexPositions();
+		TEdgeAttributesConstRef<bool> SourceEdgeHardnesses = SourceAttributes.GetEdgeHardnesses();
+		TPolygonGroupAttributesConstRef<FName> SourceImportedMaterialSlotNames = SourceAttributes.GetPolygonGroupMaterialSlotNames();
+		TVertexInstanceAttributesConstRef<FVector> SourceVertexInstanceNormals = SourceAttributes.GetVertexInstanceNormals();
+		TVertexInstanceAttributesConstRef<FVector> SourceVertexInstanceTangents = SourceAttributes.GetVertexInstanceTangents();
+		TVertexInstanceAttributesConstRef<float> SourceVertexInstanceBinormalSigns = SourceAttributes.GetVertexInstanceBinormalSigns();
+		TVertexInstanceAttributesConstRef<FVector4> SourceVertexInstanceColors = SourceAttributes.GetVertexInstanceColors();
+		TVertexInstanceAttributesConstRef<FVector2D> SourceVertexInstanceUVs = SourceAttributes.GetVertexInstanceUVs();
+
+		//PolygonGroups
+		if (AppendSettings.PolygonGroupsDelegate.IsBound())
 		{
-			FPolygonGroupID TargetMatchingID = INDEX_NONE;
-			for (FPolygonGroupID TargetPolygonGroupID : TargetMesh.PolygonGroups().GetElementIDs())
+			AppendSettings.PolygonGroupsDelegate.Execute(SourceMesh, TargetMesh, RemapPolygonGroup);
+		}
+		else
+		{			
+			for (FPolygonGroupID SourcePolygonGroupID : SourceMesh.PolygonGroups().GetElementIDs())
 			{
-				if (SourceImportedMaterialSlotNames[SourcePolygonGroupID] == TargetImportedMaterialSlotNames[TargetPolygonGroupID])
+				FPolygonGroupID TargetMatchingID = INDEX_NONE;
+				for (FPolygonGroupID TargetPolygonGroupID : TargetMesh.PolygonGroups().GetElementIDs())
 				{
-					TargetMatchingID = TargetPolygonGroupID;
-					break;
+					if (SourceImportedMaterialSlotNames[SourcePolygonGroupID] == TargetImportedMaterialSlotNames[TargetPolygonGroupID])
+					{
+						TargetMatchingID = TargetPolygonGroupID;
+						break;
+					}
 				}
+				if (TargetMatchingID == INDEX_NONE)
+				{
+					TargetMatchingID = TargetMesh.CreatePolygonGroup();
+					TargetImportedMaterialSlotNames[TargetMatchingID] = SourceImportedMaterialSlotNames[SourcePolygonGroupID];
+				}
+				RemapPolygonGroup.Add(SourcePolygonGroupID, TargetMatchingID);
 			}
-			if (TargetMatchingID == INDEX_NONE)
-			{
-				TargetMatchingID = TargetMesh.CreatePolygonGroup();
-				TargetImportedMaterialSlotNames[TargetMatchingID] = SourceImportedMaterialSlotNames[SourcePolygonGroupID];
-			}
-			RemapPolygonGroup.Add(SourcePolygonGroupID, TargetMatchingID);
 		}
-	}
 
-	//Vertices
-	TMap<FVertexID, FVertexID> SourceVertexIDRemap;
-	SourceVertexIDRemap.Reserve(SourceMesh.Vertices().Num());
-	for (FVertexID SourceVertexID : SourceMesh.Vertices().GetElementIDs())
-	{
-		FVertexID TargetVertexID = TargetMesh.CreateVertex();
-		TargetVertexPositions[TargetVertexID] = (SourceVertexPositions[SourceVertexID] - AppendSettings.MergedAssetPivot);
-		SourceVertexIDRemap.Add(SourceVertexID, TargetVertexID);
-	}
+		FPolygonGroupID SinglePolygonGroup = TargetMesh.PolygonGroups().Num() == 1 ? TargetMesh.PolygonGroups().GetFirstValidID() : INDEX_NONE;
 
-	// Transform vertices properties
-	if (AppendSettings.MeshTransform)
-	{
-		const FTransform& Transform = AppendSettings.MeshTransform.GetValue();
-		for (const TPair<FVertexID, FVertexID>& VertexIDPair : SourceVertexIDRemap)
+		//Vertices
+		for (FVertexID SourceVertexID : SourceMesh.Vertices().GetElementIDs())
 		{
-			FVector& Position = TargetVertexPositions[VertexIDPair.Value];
-			Position = Transform.TransformPosition(Position);
-		}
-	}
-
-	//Edges
-	TMap<FEdgeID, FEdgeID> SourceEdgeIDRemap;
-	SourceEdgeIDRemap.Reserve(SourceMesh.Edges().Num());
-	for (const FEdgeID SourceEdgeID : SourceMesh.Edges().GetElementIDs())
-	{
-		const FVertexID EdgeVertex0 = SourceMesh.GetEdgeVertex(SourceEdgeID, 0);
-		const FVertexID EdgeVertex1 = SourceMesh.GetEdgeVertex(SourceEdgeID, 1);
-		FEdgeID TargetEdgeID = TargetMesh.CreateEdge(SourceVertexIDRemap[EdgeVertex0], SourceVertexIDRemap[EdgeVertex1]);
-		TargetEdgeHardnesses[TargetEdgeID] = SourceEdgeHardnesses[SourceEdgeID];
-		SourceEdgeIDRemap.Add(SourceEdgeID, TargetEdgeID);
-	}
+			FVertexID TargetVertexID = TargetMesh.CreateVertex();
+			TargetVertexPositions[TargetVertexID] = (SourceVertexPositions[SourceVertexID] - AppendSettings.MergedAssetPivot);
 
-	//VertexInstances
-	TMap<FVertexInstanceID, FVertexInstanceID> SourceVertexInstanceIDRemap;
-	SourceVertexInstanceIDRemap.Reserve(SourceMesh.VertexInstances().Num());
-	for (const FVertexInstanceID& SourceVertexInstanceID : SourceMesh.VertexInstances().GetElementIDs())
-	{
-		FVertexInstanceID TargetVertexInstanceID = TargetMesh.CreateVertexInstance(SourceVertexIDRemap[SourceMesh.GetVertexInstanceVertex(SourceVertexInstanceID)]);
-		SourceVertexInstanceIDRemap.Add(SourceVertexInstanceID, TargetVertexInstanceID);
-
-		TargetVertexInstanceNormals[TargetVertexInstanceID] = SourceVertexInstanceNormals[SourceVertexInstanceID];
-		TargetVertexInstanceTangents[TargetVertexInstanceID] = SourceVertexInstanceTangents[SourceVertexInstanceID];
-		TargetVertexInstanceBinormalSigns[TargetVertexInstanceID] = SourceVertexInstanceBinormalSigns[SourceVertexInstanceID];
+			SourceToTargetVertexID.Add(SourceVertexID, TargetVertexID);
+		}
 
-		if (AppendSettings.bMergeVertexColor)
+		// Transform vertices properties
+		if (AppendSettings.MeshTransform)
 		{
-			TargetVertexInstanceColors[TargetVertexInstanceID] = SourceVertexInstanceColors[SourceVertexInstanceID];
+			const FTransform& Transform = AppendSettings.MeshTransform.GetValue();
+			for (const TPair<FVertexID, FVertexID>& VertexIDPair : SourceToTargetVertexID)
+			{
+				FVector& Position = TargetVertexPositions[VertexIDPair.Value];
+				Position = Transform.TransformPosition(Position);
+			}
 		}
 
-		for (int32 UVChannelIndex = 0; UVChannelIndex < NumSourceUVChannels; ++UVChannelIndex)
+		//Edges
+		for (const FEdgeID SourceEdgeID : SourceMesh.Edges().GetElementIDs())
 		{
-			TargetVertexInstanceUVs.Set(TargetVertexInstanceID, UVChannelIndex, SourceVertexInstanceUVs.Get(SourceVertexInstanceID, UVChannelIndex));
+			const FVertexID EdgeVertex0 = SourceMesh.GetEdgeVertex(SourceEdgeID, 0);
+			const FVertexID EdgeVertex1 = SourceMesh.GetEdgeVertex(SourceEdgeID, 1);
+			FEdgeID TargetEdgeID = TargetMesh.CreateEdge(SourceToTargetVertexID[EdgeVertex0], SourceToTargetVertexID[EdgeVertex1]);
+			TargetEdgeHardnesses[TargetEdgeID] = SourceEdgeHardnesses[SourceEdgeID];
 		}
-	}
 
-	// Transform vertex instances properties
-	if (AppendSettings.MeshTransform)
-	{
-		const FTransform& Transform = AppendSettings.MeshTransform.GetValue();
-		bool bFlipBinormal = Transform.GetDeterminant() < 0;
-		float BinormalSignsFactor = bFlipBinormal ? -1.f : 1.f;
-		for (const TPair<FVertexInstanceID, FVertexInstanceID>& VertexInstanceIDPair : SourceVertexInstanceIDRemap)
+		//VertexInstances
+		for (const FVertexInstanceID& SourceVertexInstanceID : SourceMesh.VertexInstances().GetElementIDs())
 		{
-			FVertexInstanceID InstanceID = VertexInstanceIDPair.Value;
+			const FVertexID SourceVertexID = SourceMesh.GetVertexInstanceVertex(SourceVertexInstanceID);
+			FVertexInstanceID TargetVertexInstanceID = TargetMesh.CreateVertexInstance(SourceToTargetVertexID[SourceVertexID]);
+	
+			TargetVertexInstanceNormals[TargetVertexInstanceID] = SourceVertexInstanceNormals[SourceVertexInstanceID];
+			TargetVertexInstanceTangents[TargetVertexInstanceID] = SourceVertexInstanceTangents[SourceVertexInstanceID];
+			TargetVertexInstanceBinormalSigns[TargetVertexInstanceID] = SourceVertexInstanceBinormalSigns[SourceVertexInstanceID];
 
-			FVector& Normal = TargetVertexInstanceNormals[InstanceID];
-			Normal = Transform.TransformVectorNoScale(Normal);
+			if (AppendSettings.bMergeVertexColor)
+			{
+				TargetVertexInstanceColors[TargetVertexInstanceID] = SourceVertexInstanceColors[SourceVertexInstanceID];
+			}
 
-			FVector& Tangent = TargetVertexInstanceTangents[InstanceID];
-			Tangent = Transform.TransformVectorNoScale(Tangent);
+			for (int32 UVChannelIndex = 0; UVChannelIndex < MaxNumUVChannels && UVChannelIndex < SourceVertexInstanceUVs.GetNumChannels(); ++UVChannelIndex)
+			{
+				TargetVertexInstanceUVs.Set(TargetVertexInstanceID, UVChannelIndex, SourceVertexInstanceUVs.Get(SourceVertexInstanceID, UVChannelIndex));
+			}
 
-			TargetVertexInstanceBinormalSigns[InstanceID] *= BinormalSignsFactor;
+			SourceToTargetVertexInstanceID.Add(SourceVertexInstanceID, TargetVertexInstanceID);
 		}
-	}
 
-	//Polygons
-	for (const FPolygonID SourcePolygonID : SourceMesh.Polygons().GetElementIDs())
-	{
-		TArray<FVertexInstanceID, TInlineAllocator<4>> PerimeterVertexInstanceIDs = SourceMesh.GetPolygonVertexInstances<TInlineAllocator<4>>(SourcePolygonID);
-		const FPolygonGroupID PolygonGroupID = SourceMesh.GetPolygonPolygonGroup(SourcePolygonID);
-		//Find the polygonGroupID
-		FPolygonGroupID TargetPolygonGroupID = RemapPolygonGroup[PolygonGroupID];
+		// Transform vertex instances properties
+		if (AppendSettings.MeshTransform)
+		{
+			const FTransform& Transform = AppendSettings.MeshTransform.GetValue();
+			bool bFlipBinormal = Transform.GetDeterminant() < 0;
+			float BinormalSignsFactor = bFlipBinormal ? -1.f : 1.f;
+			for (const TPair<FVertexInstanceID, FVertexInstanceID>& VertexInstanceIDPair : SourceToTargetVertexInstanceID)
+			{
+				FVertexInstanceID InstanceID = VertexInstanceIDPair.Value;
+
+				FVector& Normal = TargetVertexInstanceNormals[InstanceID];
+				Normal = Transform.TransformVectorNoScale(Normal);
+
+				FVector& Tangent = TargetVertexInstanceTangents[InstanceID];
+				Tangent = Transform.TransformVectorNoScale(Tangent);
+
+				TargetVertexInstanceBinormalSigns[InstanceID] *= BinormalSignsFactor;
+			}
+		}
 
-		int32 PolygonVertexCount = PerimeterVertexInstanceIDs.Num();
-		TArray<FVertexInstanceID> VertexInstanceIDs;
-		VertexInstanceIDs.Reserve(PolygonVertexCount);
-		for (const FVertexInstanceID VertexInstanceID : PerimeterVertexInstanceIDs)
+		// Triangles
+		for (const FTriangleID SourceTriangleID : SourceMesh.Triangles().GetElementIDs())
 		{
-			VertexInstanceIDs.Add(SourceVertexInstanceIDRemap[VertexInstanceID]);
+			TArrayView<const FVertexInstanceID> TriangleVertexInstanceIDs = SourceMesh.GetTriangleVertexInstances(SourceTriangleID);
+				
+			//Find the polygonGroupID
+			FPolygonGroupID TargetPolygonGroupID = SinglePolygonGroup != INDEX_NONE ? SinglePolygonGroup : RemapPolygonGroup[SourceMesh.GetTrianglePolygonGroup(SourceTriangleID)];
+
+			TArray<FVertexInstanceID, TInlineAllocator<3>> VertexInstanceIDs;
+			VertexInstanceIDs.Reserve(3);
+			for (const FVertexInstanceID VertexInstanceID : TriangleVertexInstanceIDs)
+			{
+				VertexInstanceIDs.Add(SourceToTargetVertexInstanceID[VertexInstanceID]);
+			}
+			// Insert a triangle into the mesh
+			TargetMesh.CreateTriangle(TargetPolygonGroupID, VertexInstanceIDs);
 		}
-		// Insert a polygon into the mesh
-		const FPolygonID TargetPolygonID = TargetMesh.CreatePolygon(TargetPolygonGroupID, VertexInstanceIDs);
 	}
-}
 
+	TargetMesh.ResumeVertexInstanceIndexing();
+	TargetMesh.ResumeEdgeIndexing();
+	TargetMesh.ResumePolygonIndexing();
+	TargetMesh.ResumePolygonGroupIndexing();
+	TargetMesh.ResumeUVIndexing();
+}
 
+void FStaticMeshOperations::AppendMeshDescription(const FMeshDescription& SourceMesh, FMeshDescription& TargetMesh, FAppendSettings& AppendSettings)
+{
+	AppendMeshDescriptions({ &SourceMesh }, TargetMesh, AppendSettings);
+}
 
 //////////////////////////////////////////////////////////////////////////
 // Normals tangents and Bi-normals
diff --git a/Engine/Source/Runtime/StaticMeshDescription/Public/StaticMeshOperations.h b/Engine/Source/Runtime/StaticMeshDescription/Public/StaticMeshOperations.h
--- a/Engine/Source/Runtime/StaticMeshDescription/Public/StaticMeshOperations.h
+++ b/Engine/Source/Runtime/StaticMeshDescription/Public/StaticMeshOperations.h
@@ -84,6 +84,7 @@ public:
 	static void ConvertFromRawMesh(const FRawMesh& SourceRawMesh, FMeshDescription& DestinationMeshDescription, const TMap<int32, FName>& MaterialMap, bool bSkipNormalsAndTangents = false);
 
 	static void AppendMeshDescription(const FMeshDescription& SourceMesh, FMeshDescription& TargetMesh, FAppendSettings& AppendSettings);
+	static void AppendMeshDescriptions(const TArray<const FMeshDescription*>& SourceMeshes, FMeshDescription& TargetMesh, FAppendSettings& AppendSettings);
 
 	static void AreNormalsAndTangentsValid(const FMeshDescription& MeshDescription, bool& bHasInvalidNormals, bool& bHasInvalidTangents);
 
