Do not assume that the static mesh triangles array will be compact prior to calculating triangle face normals/tangents.

diff --git a/Engine/Source/Runtime/StaticMeshDescription/Private/StaticMeshOperations.cpp b/Engine/Source/Runtime/StaticMeshDescription/Private/StaticMeshOperations.cpp
--- a/Engine/Source/Runtime/StaticMeshDescription/Private/StaticMeshOperations.cpp
+++ b/Engine/Source/Runtime/StaticMeshDescription/Private/StaticMeshOperations.cpp
@@ -249,7 +249,11 @@ void FStaticMeshOperations::ComputeTriangleTangentsAndNormals(FMeshDescription&
 
 	// Check that the mesh description is compact
 	const int32 NumTriangles = MeshDescription.Triangles().Num();
-	check(MeshDescription.Triangles().GetArraySize() == NumTriangles);
+	if (MeshDescription.Triangles().GetArraySize() != NumTriangles)
+	{
+		FElementIDRemappings Remappings;
+		MeshDescription.Compact(Remappings);
+	}
 
 	// Split work in batch to reduce call overhead
 	const int32 BatchSize = 8 * 1024;
