Optimized FStaticMeshOperations::GenerateUniqueUVsForStaticMesh()
* Create a mesh view instead of working with a copy of the FMeshDescription

diff --git a/Engine/Source/Runtime/MeshUtilitiesCommon/Private/OverlappingCorners.cpp b/Engine/Source/Runtime/MeshUtilitiesCommon/Private/OverlappingCorners.cpp
--- a/Engine/Source/Runtime/MeshUtilitiesCommon/Private/OverlappingCorners.cpp
+++ b/Engine/Source/Runtime/MeshUtilitiesCommon/Private/OverlappingCorners.cpp
@@ -3,16 +3,44 @@
 #include "OverlappingCorners.h"
 #include "MeshUtilitiesCommon.h"
 
+// Create a mesh view from arrays of vertices & indices
+struct FOverlappingCornersArrayMeshView : FLayoutUV::IMeshView
+{
+	const TArray<FVector>& Vertices;
+	const TArray<uint32>& Indices;
+
+	FOverlappingCornersArrayMeshView(const TArray<FVector>& InVertices, const TArray<uint32>& InIndices)
+		: Vertices(InVertices)
+		, Indices(InIndices)
+	{
+	}
+
+	// IMeshView Interface --
+	virtual uint32      GetNumIndices() const override { return Indices.Num(); }
+	virtual FVector     GetPosition(uint32 Index) const override { return Vertices[Indices[Index]]; }
+
+	virtual FVector     GetNormal(uint32 Index) const override { check(0); return FVector(); }
+	virtual FVector2D   GetInputTexcoord(uint32 Index) const override { check(0); return FVector2D(); }
+	virtual void        InitOutputTexcoords(uint32 Num) override { check(0); }
+	virtual void        SetOutputTexcoord(uint32 Index, const FVector2D& Value) override { check(0); }
+	// IMeshView Interface --
+};
+
 FOverlappingCorners::FOverlappingCorners(const TArray<FVector>& InVertices, const TArray<uint32>& InIndices, float ComparisonThreshold)
+	: FOverlappingCorners(FOverlappingCornersArrayMeshView(InVertices, InIndices), ComparisonThreshold)
+{
+}
+
+FOverlappingCorners::FOverlappingCorners(const FLayoutUV::IMeshView& MeshView, float ComparisonThreshold)
 {
-	const int32 NumWedges = InIndices.Num();
+	const int32 NumWedges = MeshView.GetNumIndices();
 
 	// Create a list of vertex Z/index pairs
 	TArray<FIndexAndZ> VertIndexAndZ;
 	VertIndexAndZ.Reserve(NumWedges);
 	for (int32 WedgeIndex = 0; WedgeIndex < NumWedges; WedgeIndex++)
 	{
-		new(VertIndexAndZ)FIndexAndZ(WedgeIndex, InVertices[InIndices[WedgeIndex]]);
+		new(VertIndexAndZ)FIndexAndZ(WedgeIndex, MeshView.GetPosition(WedgeIndex));
 	}
 
 	// Sort the vertices by z value
@@ -29,8 +57,8 @@ FOverlappingCorners::FOverlappingCorners(const TArray<FVector>& InVertices, cons
 			if (FMath::Abs(VertIndexAndZ[j].Z - VertIndexAndZ[i].Z) > ComparisonThreshold)
 				break; // can't be any more dups
 
-			const FVector& PositionA = InVertices[InIndices[VertIndexAndZ[i].Index]];
-			const FVector& PositionB = InVertices[InIndices[VertIndexAndZ[j].Index]];
+			const FVector& PositionA = MeshView.GetPosition(VertIndexAndZ[i].Index);
+			const FVector& PositionB = MeshView.GetPosition(VertIndexAndZ[j].Index);
 
 			if (PointsEqual(PositionA, PositionB, ComparisonThreshold))
 			{
diff --git a/Engine/Source/Runtime/MeshUtilitiesCommon/Public/OverlappingCorners.h b/Engine/Source/Runtime/MeshUtilitiesCommon/Public/OverlappingCorners.h
--- a/Engine/Source/Runtime/MeshUtilitiesCommon/Public/OverlappingCorners.h
+++ b/Engine/Source/Runtime/MeshUtilitiesCommon/Public/OverlappingCorners.h
@@ -3,6 +3,7 @@
 #pragma once
 
 #include "CoreMinimal.h"
+#include "LayoutUV.h"
 
 /**
 * Container to hold overlapping corners. For a vertex, lists all the overlapping vertices
@@ -12,6 +13,7 @@ struct MESHUTILITIESCOMMON_API FOverlappingCorners
 	FOverlappingCorners() {}
 
 	FOverlappingCorners(const TArray<FVector>& InVertices, const TArray<uint32>& InIndices, float ComparisonThreshold);
+	FOverlappingCorners(const FLayoutUV::IMeshView& MeshView, float ComparisonThreshold);
 
 	/* Resets, pre-allocates memory, marks all indices as not overlapping in preperation for calls to Add() */
 	void Init(int32 NumIndices);
diff --git a/Engine/Source/Runtime/StaticMeshDescription/Private/StaticMeshOperations.cpp b/Engine/Source/Runtime/StaticMeshDescription/Private/StaticMeshOperations.cpp
--- a/Engine/Source/Runtime/StaticMeshDescription/Private/StaticMeshOperations.cpp
+++ b/Engine/Source/Runtime/StaticMeshDescription/Private/StaticMeshOperations.cpp
@@ -1916,118 +1916,173 @@ bool FStaticMeshOperations::CreateLightMapUVLayout(FMeshDescription& MeshDescrip
 	return bPackSuccess;
 }
 
-bool FStaticMeshOperations::GenerateUniqueUVsForStaticMesh(const FMeshDescription& MeshDescription, int32 TextureResolution, bool bMergeIdenticalMaterials, TArray<FVector2D>& OutTexCoords)
+// Mesh view that will expose only unique UVs if bMergeIdenticalMaterials is provided
+struct FUniqueUVMeshDescriptionView final : FLayoutUV::IMeshView
 {
-	// Create a copy of original mesh (only copy necessary data)
-	FMeshDescription DuplicateMeshDescription(MeshDescription);
+	const FMeshDescription& MeshDescription;
 
+	TVertexAttributesConstRef<FVector> Positions;
+	TVertexInstanceAttributesConstRef<FVector> Normals;
+	TVertexInstanceAttributesConstRef<FVector2D> TexCoords;
 
-	//Make sure we have a destination UV TextureCoordinnate
-	{
-		TVertexInstanceAttributesRef<FVector2D> DuplicateVertexInstanceUVs = DuplicateMeshDescription.VertexInstanceAttributes().GetAttributesRef<FVector2D>(MeshAttribute::VertexInstance::TextureCoordinate);
-		if (DuplicateVertexInstanceUVs.GetNumChannels() < 2)
-		{
-			DuplicateVertexInstanceUVs.SetNumChannels(2);
-		}
-	}
+	uint32 NumIndices = 0;
+
+	TArray<int32>						RemapVerts;
+	TMap<FVertexInstanceID, int32>		InstanceToIndex;
+	TArray<FVertexInstanceID>			UniqueVerts;
+	TArray<FVector2D>					OutputTexCoords;
 
-	TMap<FVertexInstanceID, FVertexInstanceID> RemapVertexInstance;
-	//Remove the identical material
-	if (bMergeIdenticalMaterials)
+	FUniqueUVMeshDescriptionView(const FMeshDescription& InMeshDescription, bool bMergeIdenticalMaterials)
+		: MeshDescription(InMeshDescription)
 	{
-		TVertexInstanceAttributesConstRef<FVector2D> VertexInstanceUVs = DuplicateMeshDescription.VertexInstanceAttributes().GetAttributesRef<FVector2D>(MeshAttribute::VertexInstance::TextureCoordinate);
-		TVertexInstanceAttributesConstRef<FVector4> VertexColors = DuplicateMeshDescription.VertexInstanceAttributes().GetAttributesRef<FVector4>(MeshAttribute::VertexInstance::Color);
+		FStaticMeshConstAttributes Attributes(MeshDescription);
+		Positions = Attributes.GetVertexPositions();
+		Normals = Attributes.GetVertexInstanceNormals();
+		TexCoords = Attributes.GetVertexInstanceUVs();
+
+		TVertexInstanceAttributesConstRef<FVector4> VertexColors = Attributes.GetVertexInstanceColors();
 
-		TArray<FPolygonID> ToDeletePolygons;
-		RemapVertexInstance.Reserve(DuplicateMeshDescription.VertexInstances().Num());
+		NumIndices = MeshDescription.VertexInstances().Num();
 
-		int32 NumPolygons = DuplicateMeshDescription.Polygons().Num();
+		OutputTexCoords.SetNumUninitialized(NumIndices);
+		UniqueVerts.Reserve(NumIndices);
 
-		TMap<uint32, FPolygonID> UniquePolygons;
-		UniquePolygons.Reserve(NumPolygons);
-		ToDeletePolygons.Reserve(NumPolygons);
+		TMap<uint32, FTriangleID> UniqueTriangles;
 
-		auto HashAttribute = [](FVertexInstanceID InVertexInstanceID, auto InAttributeArrayRef, int32& OutPolyHash)
+		if (bMergeIdenticalMaterials)
+		{
+			RemapVerts.Reserve(NumIndices);
+			UniqueTriangles.Reserve(MeshDescription.Triangles().Num());
+		}
+
+		auto HashAttribute = [](FVertexInstanceID InVertexInstanceID, auto InAttributeArrayRef, int32& TriangleHash)
 		{
 			for (int32 Channel = 0; Channel < InAttributeArrayRef.GetNumChannels(); ++Channel)
 			{
 				for (const auto& Element : InAttributeArrayRef.GetArrayView(InVertexInstanceID, Channel))
 				{
-					OutPolyHash = HashCombine(OutPolyHash, GetTypeHash(Element));
+					TriangleHash = HashCombine(TriangleHash, GetTypeHash(Element));
 				}
 			}
 		};
 
-		for (FPolygonID RefPolygonID : DuplicateMeshDescription.Polygons().GetElementIDs())
+		for (const FTriangleID TriangleID : MeshDescription.Triangles().GetElementIDs())
 		{
-			const FPolygonGroupID RefPolygonGroupID = DuplicateMeshDescription.GetPolygonPolygonGroup(RefPolygonID);
-			TArray<FVertexInstanceID, TInlineAllocator<4>> RefVertexInstances = DuplicateMeshDescription.GetPolygonVertexInstances<TInlineAllocator<4>>(RefPolygonID);
+			const FPolygonGroupID RefPolygonGroupID = MeshDescription.GetTrianglePolygonGroup(TriangleID);
+			TArrayView<const FVertexInstanceID> VertexInstances = MeshDescription.GetTriangleVertexInstances(TriangleID);
 
-			int32 PolyHash = GetTypeHash(RefPolygonGroupID);
-			for (FVertexInstanceID RefVertexInstanceID : RefVertexInstances)
+			bool bUnique = true;
+			if (bMergeIdenticalMaterials)
 			{
-				// Compute hash based on UVs & vertices colors
-				HashAttribute(RefVertexInstanceID, VertexInstanceUVs, PolyHash);
-				HashAttribute(RefVertexInstanceID, VertexColors, PolyHash);
-			}
+				int32 TriangleHash = GetTypeHash(RefPolygonGroupID);
+				for (const FVertexInstanceID VertexInstanceID : VertexInstances)
+				{
+					// Compute hash based on UVs & vertices colors
+					HashAttribute(VertexInstanceID, TexCoords, TriangleHash);
+					HashAttribute(VertexInstanceID, VertexColors, TriangleHash);
+				}
 
-			FPolygonID& UniquePoly = UniquePolygons.FindOrAdd(PolyHash);
-			if (UniquePoly == INDEX_NONE)
-			{
-				UniquePoly = RefPolygonID;
-				for (FVertexInstanceID RefVertexInstanceID : RefVertexInstances)
+				FTriangleID& UniqueTriangle = UniqueTriangles.FindOrAdd(TriangleHash);
+				if (UniqueTriangle != INDEX_NONE)
 				{
-					RemapVertexInstance.Add(RefVertexInstanceID, RefVertexInstanceID);
+					for (const FVertexInstanceID UniqueVertexInstance : MeshDescription.GetTriangleVertexInstances(UniqueTriangle))
+					{
+						RemapVerts.Add(InstanceToIndex[UniqueVertexInstance]);
+					}
+
+					bUnique = false;
+				}
+				else
+				{
+					UniqueTriangle = TriangleID;
 				}
 			}
-			else
+			
+			if (bUnique)
 			{
-				TArray<FVertexInstanceID, TInlineAllocator<4>> TestVertexInstances = DuplicateMeshDescription.GetPolygonVertexInstances<TInlineAllocator<4>>(UniquePoly);
-				int32 VertexInstanceIndex = 0;
-				for (FVertexInstanceID RefVertexInstanceID : RefVertexInstances)
+				for (const FVertexInstanceID VertexInstanceID : VertexInstances)
 				{
-					RemapVertexInstance.Add(RefVertexInstanceID, TestVertexInstances[VertexInstanceIndex]);
-					VertexInstanceIndex++;
+					if (bMergeIdenticalMaterials)
+					{
+						InstanceToIndex.Emplace(VertexInstanceID, UniqueVerts.Num());
+						RemapVerts.Add(UniqueVerts.Num());
+					}
+
+					UniqueVerts.Add(VertexInstanceID);
 				}
-				ToDeletePolygons.Add(RefPolygonID);
 			}
 		}
+	}
 
-		//Delete polygons
-		if (ToDeletePolygons.Num() > 0)
-		{
-			DuplicateMeshDescription.DeletePolygons(ToDeletePolygons);
+	uint32 GetNumIndices() const override 
+	{
+		return UniqueVerts.Num();
+	}
 
-			//Avoid compacting the DuplicateMeshDescription, since the remap of the VertexInstanceID will not be good anymore
+	FVector GetPosition(uint32 Index) const override
+	{
+		FVertexID VertexID = MeshDescription.GetVertexInstanceVertex(UniqueVerts[Index]);
+		return Positions[VertexID];
+	}
+
+	FVector GetNormal(uint32 Index) const override
+	{
+		return Normals[UniqueVerts[Index]];
+	}
+
+	FVector2D GetInputTexcoord(uint32 Index) const override
+	{
+		return TexCoords.Get(UniqueVerts[Index], 0);
+	}
+
+	void InitOutputTexcoords(uint32 Num) override
+	{
+		check(Num == UniqueVerts.Num());
+	}
+
+	void SetOutputTexcoord(uint32 Index, const FVector2D& Value) override
+	{
+		OutputTexCoords[Index] = Value;
+	}
+
+	TArray<FVector2D> RetrievePackedUVs()
+	{
+		if (RemapVerts.IsEmpty())
+		{ 
+			return OutputTexCoords;
 		}
+		
+		TArray<FVector2D> Out;
+		Out.Reserve(NumIndices);
+
+		for (uint32 i = 0; i < NumIndices; ++i)
+		{
+			Out.Add(OutputTexCoords[RemapVerts[i]]);
+		}
+
+		return Out;
 	}
+};
+
+bool FStaticMeshOperations::GenerateUniqueUVsForStaticMesh(const FMeshDescription& MeshDescription, int32 TextureResolution, bool bMergeIdenticalMaterials, TArray<FVector2D>& OutTexCoords)
+{
+	TRACE_CPUPROFILER_EVENT_SCOPE(FStaticMeshOperations::GenerateUniqueUVsForStaticMesh)
+
+	FUniqueUVMeshDescriptionView MeshDescriptionView(MeshDescription, bMergeIdenticalMaterials);
+
 	// Find overlapping corners for UV generator. Allow some threshold - this should not produce any error in a case if resulting
 	// mesh will not merge these vertices.
-	FOverlappingCorners OverlappingCorners;
-	FindOverlappingCorners(OverlappingCorners, DuplicateMeshDescription, THRESH_POINTS_ARE_SAME);
+	FOverlappingCorners OverlappingCorners(MeshDescriptionView, THRESH_POINTS_ARE_SAME);
 
 	// Generate new UVs
-	FLayoutUVMeshDescriptionView DuplicateMeshDescriptionView(DuplicateMeshDescription, 0, 1);
-	FLayoutUV Packer(DuplicateMeshDescriptionView);
-	Packer.FindCharts(OverlappingCorners);
+	FLayoutUV Packer(MeshDescriptionView);
+	int32 NumCharts = Packer.FindCharts(OverlappingCorners);
 
 	bool bPackSuccess = Packer.FindBestPacking(FMath::Clamp(TextureResolution / 4, 32, 512));
 	if (bPackSuccess)
 	{
 		Packer.CommitPackedUVs();
-		TVertexInstanceAttributesConstRef<FVector2D> DupVertexInstanceUVs = DuplicateMeshDescription.VertexInstanceAttributes().GetAttributesRef<FVector2D>(MeshAttribute::VertexInstance::TextureCoordinate);
-		TVertexInstanceAttributesConstRef<FVector2D> VertexInstanceUVs = MeshDescription.VertexInstanceAttributes().GetAttributesRef<FVector2D>(MeshAttribute::VertexInstance::TextureCoordinate);
-		// Save generated UVs
-		check(DupVertexInstanceUVs.GetNumChannels() > 1);
-		OutTexCoords.AddZeroed(VertexInstanceUVs.GetNumElements());
-		int32 TextureCoordIndex = 0;
-		for (const FVertexInstanceID& VertexInstanceID : MeshDescription.VertexInstances().GetElementIDs())
-		{
-			FVertexInstanceID RemapID = bMergeIdenticalMaterials ? RemapVertexInstance[VertexInstanceID] : VertexInstanceID;
-			// Save generated UVs
-			OutTexCoords[TextureCoordIndex] = DupVertexInstanceUVs.Get(RemapID, 1);	// UV1
-			TextureCoordIndex++;
-		}
+		OutTexCoords = MeshDescriptionView.RetrievePackedUVs();
 	}
 
 	return bPackSuccess;
