StaticMeshDescription: Change loop variables to references and/or copies to resolve errors in Mac builds

diff --git a/Engine/Source/Runtime/StaticMeshDescription/Private/StaticMeshOperations.cpp b/Engine/Source/Runtime/StaticMeshDescription/Private/StaticMeshOperations.cpp
--- a/Engine/Source/Runtime/StaticMeshDescription/Private/StaticMeshOperations.cpp
+++ b/Engine/Source/Runtime/StaticMeshDescription/Private/StaticMeshOperations.cpp
@@ -40,14 +40,14 @@ static bool GetPolygonTangentsAndNormals(FMeshDescription& MeshDescription,
 
 	// Calculate the center of this polygon
 	TArray<FVertexInstanceID, TInlineAllocator<4>> VertexInstanceIDs = MeshDescription.GetPolygonVertexInstances<TInlineAllocator<4>>(PolygonID);
-	for (const FVertexInstanceID VertexInstanceID : VertexInstanceIDs)
+	for (const FVertexInstanceID& VertexInstanceID : VertexInstanceIDs)
 	{
 		Center += VertexPositions[MeshDescription.GetVertexInstanceVertex(VertexInstanceID)];
 	}
 	Center /= float(VertexInstanceIDs.Num());
 
 	float AdjustedComparisonThreshold = FMath::Max(ComparisonThreshold, MIN_flt);
-	for (const FTriangleID TriangleID : MeshDescription.GetPolygonTriangles(PolygonID))
+	for (const FTriangleID& TriangleID : MeshDescription.GetPolygonTriangles(PolygonID))
 	{
 		TArrayView<const FVertexInstanceID> TriangleVertexInstances = MeshDescription.GetTriangleVertexInstances(TriangleID);
 		const FVertexID VertexID0 = MeshDescription.GetVertexInstanceVertex(TriangleVertexInstances[0]);
@@ -137,7 +137,7 @@ void FStaticMeshOperations::ComputePolygonTangentsAndNormals(FMeshDescription& M
 
 	TArray<FPolygonID> PolygonIDs;
 	PolygonIDs.Reserve(MeshDescription.Polygons().Num());
-	for (const FPolygonID& Polygon : MeshDescription.Polygons().GetElementIDs())
+	for (const FPolygonID Polygon : MeshDescription.Polygons().GetElementIDs())
 	{
 		PolygonIDs.Add(Polygon);
 	}
@@ -340,9 +340,9 @@ void FStaticMeshOperations::DetermineEdgeHardnessesFromVertexInstanceNormals(FMe
 			UniqueVertexInstanceIDs.Reset(ReservedElements);
 
 			// Get a list of all vertex instances for this vertex which form part of any polygon connected to the edge
-			for (const FVertexInstanceID VertexInstanceID : MeshDescription.GetVertexVertexInstanceIDs(VertexID))
+			for (const FVertexInstanceID& VertexInstanceID : MeshDescription.GetVertexVertexInstanceIDs(VertexID))
 			{
-				for (const FPolygonID PolygonID : MeshDescription.GetVertexInstanceConnectedPolygons<TInlineAllocator<8>>(VertexInstanceID))
+				for (const FPolygonID& PolygonID : MeshDescription.GetVertexInstanceConnectedPolygons<TInlineAllocator<8>>(VertexInstanceID))
 				{
 					if (ConnectedPolygonIDs.Contains(PolygonID))
 					{
@@ -542,8 +542,8 @@ void FStaticMeshOperations::ConvertHardEdgesToSmoothGroup(const FMeshDescription
 	int32 TriangleIndex = 0;
 	for (const FPolygonID PolygonID : SourceMeshDescription.Polygons().GetElementIDs())
 	{
-		uint32 PolygonSmoothValue = PolygonSmoothGroup[PolygonID];
-		for (const FTriangleID TriangleID : SourceMeshDescription.GetPolygonTriangles(PolygonID))
+		const uint32 PolygonSmoothValue = PolygonSmoothGroup[PolygonID];
+		for (int32 i = 0, Num = SourceMeshDescription.GetPolygonTriangles(PolygonID).Num(); i < Num; ++i)
 		{
 			FaceSmoothingMasks[TriangleIndex++] = PolygonSmoothValue;
 		}
@@ -625,7 +625,7 @@ void FStaticMeshOperations::ConvertToRawMesh(const FMeshDescription& SourceMeshD
 	TArray<int32> RemapVerts;
 	RemapVerts.AddZeroed(SourceMeshDescription.Vertices().GetArraySize());
 	int32 VertexIndex = 0;
-	for (const FVertexID& VertexID : SourceMeshDescription.Vertices().GetElementIDs())
+	for (const FVertexID VertexID : SourceMeshDescription.Vertices().GetElementIDs())
 	{
 		DestinationRawMesh.VertexPositions[VertexIndex] = VertexPositions[VertexID];
 		RemapVerts[VertexID.GetValue()] = VertexIndex;
@@ -660,7 +660,7 @@ void FStaticMeshOperations::ConvertToRawMesh(const FMeshDescription& SourceMeshD
 		const FPolygonGroupID& PolygonGroupID = SourceMeshDescription.GetPolygonPolygonGroup(PolygonID);
 		int32 PolygonIDValue = PolygonID.GetValue();
 		TArrayView<const FTriangleID> TriangleIDs = SourceMeshDescription.GetPolygonTriangles(PolygonID);
-		for (const FTriangleID TriangleID : TriangleIDs)
+		for (const FTriangleID& TriangleID : TriangleIDs)
 		{
 			if (MaterialMap.Num() > 0 && MaterialMap.Contains(PolygonGroupMaterialSlotName[PolygonGroupID]))
 			{
@@ -877,7 +877,7 @@ void FStaticMeshOperations::ConvertFromRawMesh(const FRawMesh& SourceRawMesh, FM
 		else if (MaterialMap.Num() > 0 && MaterialMap.Contains(MaterialIndex))
 		{
 			PolygonGroupImportedMaterialSlotName = MaterialMap[MaterialIndex];
-			for (const FPolygonGroupID& SearchPolygonGroupID : DestinationMeshDescription.PolygonGroups().GetElementIDs())
+			for (const FPolygonGroupID SearchPolygonGroupID : DestinationMeshDescription.PolygonGroups().GetElementIDs())
 			{
 				if (PolygonGroupImportedMaterialSlotNames[SearchPolygonGroupID] == PolygonGroupImportedMaterialSlotName)
 				{
@@ -1097,7 +1097,7 @@ void FStaticMeshOperations::AppendMeshDescriptions(const TArray<const FMeshDescr
 		}
 
 		//VertexInstances
-		for (const FVertexInstanceID& SourceVertexInstanceID : SourceMesh.VertexInstances().GetElementIDs())
+		for (const FVertexInstanceID SourceVertexInstanceID : SourceMesh.VertexInstances().GetElementIDs())
 		{
 			const FVertexID SourceVertexID = SourceMesh.GetVertexInstanceVertex(SourceVertexInstanceID);
 			FVertexInstanceID TargetVertexInstanceID = TargetMesh.CreateVertexInstance(SourceToTargetVertexID[SourceVertexID]);
@@ -1149,7 +1149,7 @@ void FStaticMeshOperations::AppendMeshDescriptions(const TArray<const FMeshDescr
 
 			TArray<FVertexInstanceID, TInlineAllocator<3>> VertexInstanceIDs;
 			VertexInstanceIDs.Reserve(3);
-			for (const FVertexInstanceID VertexInstanceID : TriangleVertexInstanceIDs)
+			for (const FVertexInstanceID& VertexInstanceID : TriangleVertexInstanceIDs)
 			{
 				VertexInstanceIDs.Add(SourceToTargetVertexInstanceID[VertexInstanceID]);
 			}
@@ -1416,16 +1416,16 @@ void FStaticMeshOperations::ComputeTangentsAndNormals(FMeshDescription& MeshDesc
 
 				bool bPointHasAllTangents = true;
 				// Fill the VertexInfoMap
-				for (const FEdgeID EdgeID : MeshDescription.GetVertexConnectedEdgeIDs(VertexID))
+				for (const FEdgeID& EdgeID : MeshDescription.GetVertexConnectedEdgeIDs(VertexID))
 				{
-					for (const FTriangleID TriangleID : MeshDescription.GetEdgeConnectedTriangleIDs(EdgeID))
+					for (const FTriangleID& TriangleID : MeshDescription.GetEdgeConnectedTriangleIDs(EdgeID))
 					{
 						FVertexInfo& VertexInfo = VertexInfoMap.FindOrAdd(TriangleID);
 						int32 EdgeIndex = VertexInfo.EdgeIDs.AddUnique(EdgeID);
 						if (VertexInfo.TriangleID == INDEX_NONE)
 						{
 							VertexInfo.TriangleID = TriangleID;
-							for (const FVertexInstanceID VertexInstanceID : MeshDescription.GetTriangleVertexInstances(TriangleID))
+							for (const FVertexInstanceID& VertexInstanceID : MeshDescription.GetTriangleVertexInstances(TriangleID))
 							{
 								if (MeshDescription.GetVertexInstanceVertex(VertexInstanceID) == VertexID)
 								{
@@ -1476,14 +1476,14 @@ void FStaticMeshOperations::ComputeTangentsAndNormals(FMeshDescription& MeshDesc
 						FVertexInfo& CurrentVertexInfo = VertexInfoMap.FindOrAdd(CurrentPolygonID);
 						CurrentGroup.AddUnique(CurrentVertexInfo.TriangleID);
 						ConsumedTriangle.AddUnique(CurrentVertexInfo.TriangleID);
-						for (const FEdgeID EdgeID : CurrentVertexInfo.EdgeIDs)
+						for (const FEdgeID& EdgeID : CurrentVertexInfo.EdgeIDs)
 						{
 							if (EdgeHardnesses[EdgeID])
 							{
 								//End of the group
 								continue;
 							}
-							for (const FTriangleID TriangleID : MeshDescription.GetEdgeConnectedTriangleIDs(EdgeID))
+							for (const FTriangleID& TriangleID : MeshDescription.GetEdgeConnectedTriangleIDs(EdgeID))
 							{
 								if (TriangleID == CurrentVertexInfo.TriangleID)
 								{
@@ -1509,7 +1509,7 @@ void FStaticMeshOperations::ComputeTangentsAndNormals(FMeshDescription& MeshDesc
 					VertexInstanceInGroup.Reset();
 
 					FVector3f GroupNormal(FVector3f::ZeroVector);
-					for (const FTriangleID TriangleID : Group)
+					for (const FTriangleID& TriangleID : Group)
 					{
 						FVertexInfo& CurrentVertexInfo = VertexInfoMap.FindOrAdd(TriangleID);
 						float CornerWeight = 1.0f;
@@ -1566,7 +1566,7 @@ void FStaticMeshOperations::ComputeTangentsAndNormals(FMeshDescription& MeshDesc
 						}
 					}
 					//Apply the average NTB on all Vertex instance
-					for (const FVertexInstanceID VertexInstanceID : VertexInstanceInGroup)
+					for (const FVertexInstanceID& VertexInstanceID : VertexInstanceInGroup)
 					{
 						const FVector2f& VertexUV = VertexUVs[VertexInstanceID];	// UV0
 
@@ -1760,7 +1760,7 @@ void FStaticMeshOperations::FindOverlappingCorners(FOverlappingCorners& OutOverl
 	for (const FPolygonID PolygonID : MeshDescription.Polygons().GetElementIDs())
 	{
 		TArrayView<const FTriangleID> TriangleIDs = MeshDescription.GetPolygonTriangles(PolygonID);
-		for (const FTriangleID TriangleID : TriangleIDs)
+		for (const FTriangleID& TriangleID : TriangleIDs)
 		{
 			for (int32 Corner = 0; Corner < 3; ++Corner)
 			{
@@ -1833,7 +1833,7 @@ struct FLayoutUVMeshDescriptionView final : FLayoutUV::IMeshView
 		for (const FPolygonID PolygonID : MeshDescription.Polygons().GetElementIDs())
 		{
 			TArrayView<const FTriangleID> TriangleIDs = MeshDescription.GetPolygonTriangles(PolygonID);
-			for (const FTriangleID TriangleID : TriangleIDs)
+			for (const FTriangleID& TriangleID : TriangleIDs)
 			{
 				for (int32 Corner = 0; Corner < 3; ++Corner)
 				{
@@ -1980,7 +1980,7 @@ struct FUniqueUVMeshDescriptionView final : FLayoutUV::IMeshView
 			if (bMergeIdenticalMaterials)
 			{
 				int32 TriangleHash = GetTypeHash(RefPolygonGroupID);
-				for (const FVertexInstanceID VertexInstanceID : VertexInstances)
+				for (const FVertexInstanceID& VertexInstanceID : VertexInstances)
 				{
 					// Compute hash based on UVs & vertices colors
 					HashAttribute(VertexInstanceID, TexCoords, TriangleHash);
@@ -1990,7 +1990,7 @@ struct FUniqueUVMeshDescriptionView final : FLayoutUV::IMeshView
 				FTriangleID& UniqueTriangle = UniqueTriangles.FindOrAdd(TriangleHash);
 				if (UniqueTriangle != INDEX_NONE)
 				{
-					for (const FVertexInstanceID UniqueVertexInstance : MeshDescription.GetTriangleVertexInstances(UniqueTriangle))
+					for (const FVertexInstanceID& UniqueVertexInstance : MeshDescription.GetTriangleVertexInstances(UniqueTriangle))
 					{
 						RemapVerts.Add(InstanceToIndex[UniqueVertexInstance]);
 					}
@@ -2005,7 +2005,7 @@ struct FUniqueUVMeshDescriptionView final : FLayoutUV::IMeshView
 			
 			if (bUnique)
 			{
-				for (const FVertexInstanceID VertexInstanceID : VertexInstances)
+				for (const FVertexInstanceID& VertexInstanceID : VertexInstances)
 				{
 					if (bMergeIdenticalMaterials)
 					{
@@ -2157,7 +2157,7 @@ void FStaticMeshOperations::GeneratePlanarUV(const FMeshDescription& MeshDescrip
 	FVector3f Size = Params.Size * Params.Scale;
 	FVector3f Offset = Params.Position - Size / 2.f;
 
-	for (const FVertexInstanceID& VertexInstanceID : MeshDescription.VertexInstances().GetElementIDs())
+	for (const FVertexInstanceID VertexInstanceID : MeshDescription.VertexInstances().GetElementIDs())
 	{
 		const FVertexID VertexID = MeshDescription.GetVertexInstanceVertex(VertexInstanceID);
 		FVector3f Vertex = VertexPositions[VertexID];
@@ -2188,7 +2188,7 @@ void FStaticMeshOperations::GenerateCylindricalUV(FMeshDescription& MeshDescript
 
 	const float AngleOffset = PI; // offset to get the same result as in 3dsmax
 
-	for (const FVertexInstanceID& VertexInstanceID : MeshDescription.VertexInstances().GetElementIDs())
+	for (const FVertexInstanceID VertexInstanceID : MeshDescription.VertexInstances().GetElementIDs())
 	{
 		const FVertexID VertexID = MeshDescription.GetVertexInstanceVertex(VertexInstanceID);
 		FVector3f Vertex = VertexPositions[VertexID];
@@ -2210,7 +2210,7 @@ void FStaticMeshOperations::GenerateCylindricalUV(FMeshDescription& MeshDescript
 	}
 
 	// Fix the UV coordinates for triangles at the seam where the angle wraps around
-	for (const FPolygonID& PolygonID : MeshDescription.Polygons().GetElementIDs())
+	for (const FPolygonID PolygonID : MeshDescription.Polygons().GetElementIDs())
 	{
 		TArray<FVertexInstanceID, TInlineAllocator<4>> VertexInstances = MeshDescription.GetPolygonVertexInstances<TInlineAllocator<4>>(PolygonID);
 		int32 NumInstances = VertexInstances.Num();
@@ -2310,7 +2310,7 @@ void FStaticMeshOperations::GenerateBoxUV(const FMeshDescription& MeshDescriptio
 		FVector3f V = PlaneUVs[BestPlaneIndex].Value;
 		FVector3f Offset = Params.Position - HalfSize * (U + V);
 
-		for (const FVertexInstanceID VertexInstanceID : VertexInstances)
+		for (const FVertexInstanceID& VertexInstanceID : VertexInstances)
 		{
 			const FVertexID VertexID = MeshDescription.GetVertexInstanceVertex(VertexInstanceID);
 			FVector3f Vertex = VertexPositions[VertexID];
@@ -2382,7 +2382,7 @@ bool FStaticMeshOperations::HasVertexColor(const FMeshDescription& MeshDescripti
 	TVertexInstanceAttributesConstRef<FVector4f> VertexInstanceColors = MeshDescription.VertexInstanceAttributes().GetAttributesRef<FVector4f>(MeshAttribute::VertexInstance::Color);
 	bool bHasVertexColor = false;
 	FVector4f WhiteColor(FLinearColor::White);
-	for (const FVertexInstanceID& VertexInstanceID : MeshDescription.VertexInstances().GetElementIDs())
+	for (const FVertexInstanceID VertexInstanceID : MeshDescription.VertexInstances().GetElementIDs())
 	{
 		if (VertexInstanceColors[VertexInstanceID] != WhiteColor)
 		{
@@ -2489,7 +2489,7 @@ FSHAHash FStaticMeshOperations::ComputeSHAHash(const FMeshDescription& MeshDescr
 void FStaticMeshOperations::FlipPolygons(FMeshDescription& MeshDescription)
 {
 	TSet<FVertexInstanceID> VertexInstanceIDs;
-	for (FTriangleID TriangleID : MeshDescription.Triangles().GetElementIDs())
+	for (const FTriangleID TriangleID : MeshDescription.Triangles().GetElementIDs())
 	{
 		TArrayView<const FVertexInstanceID> TriVertInstances = MeshDescription.GetTriangleVertexInstances(TriangleID);
 		for (const FVertexInstanceID TriVertInstance : TriVertInstances)
