PR #8721: ClearNormalsAndTangentsData early return bugfix (Contributed by [REDACTED_PERSON])

diff --git a/Engine/Source/Runtime/StaticMeshDescription/Private/StaticMeshOperations.cpp b/Engine/Source/Runtime/StaticMeshDescription/Private/StaticMeshOperations.cpp
--- a/Engine/Source/Runtime/StaticMeshDescription/Private/StaticMeshOperations.cpp
+++ b/Engine/Source/Runtime/StaticMeshDescription/Private/StaticMeshOperations.cpp
@@ -1194,7 +1194,7 @@ void FStaticMeshOperations::AreNormalsAndTangentsValid(const FMeshDescription& M
 
 void ClearNormalsAndTangentsData(FMeshDescription& MeshDescription, bool bClearNormals, bool bClearTangents)
 {
-	if (!bClearNormals && bClearTangents)
+	if (!bClearNormals && !bClearTangents)
 	{
 		return;
 	}
