StaticMeshOperation: Improve triangle normal computing in case of small triangle

diff --git a/Engine/Source/Runtime/StaticMeshDescription/Private/StaticMeshOperations.cpp b/Engine/Source/Runtime/StaticMeshDescription/Private/StaticMeshOperations.cpp
--- a/Engine/Source/Runtime/StaticMeshDescription/Private/StaticMeshOperations.cpp
+++ b/Engine/Source/Runtime/StaticMeshDescription/Private/StaticMeshOperations.cpp
@@ -46,7 +46,8 @@ static bool GetPolygonTangentsAndNormals(FMeshDescription& MeshDescription,
 	}
 	Center /= float(VertexInstanceIDs.Num());
 
-	float AdjustedComparisonThreshold = FMath::Max(ComparisonThreshold, MIN_flt);
+	// GetSafeNormal compare the squareSum to the tolerance.
+	const float SquareComparisonThreshold = FMath::Max(ComparisonThreshold * ComparisonThreshold, MIN_flt);
 	for (const FTriangleID& TriangleID : MeshDescription.GetPolygonTriangles(PolygonID))
 	{
 		TArrayView<const FVertexInstanceID> TriangleVertexInstances = MeshDescription.GetTriangleVertexInstances(TriangleID);
@@ -55,8 +56,9 @@ static bool GetPolygonTangentsAndNormals(FMeshDescription& MeshDescription,
 		const FVertexID VertexID2 = MeshDescription.GetVertexInstanceVertex(TriangleVertexInstances[2]);
 
 		const FVector3f Position0 = VertexPositions[VertexID0];
-		const FVector3f DPosition1 = VertexPositions[VertexID1] - Position0;
-		const FVector3f DPosition2 = VertexPositions[VertexID2] - Position0;
+		// To avoid numerical error due to small edges, DPosition is normalized
+		const FVector3f DPosition1 = (VertexPositions[VertexID1] - Position0).GetSafeNormal(SquareComparisonThreshold);
+		const FVector3f DPosition2 = (VertexPositions[VertexID2] - Position0).GetSafeNormal(SquareComparisonThreshold);
 
 		const FVector2f UV0 = VertexUVs[TriangleVertexInstances[0]];
 		const FVector2f DUV1 = VertexUVs[TriangleVertexInstances[1]] - UV0;
@@ -64,7 +66,7 @@ static bool GetPolygonTangentsAndNormals(FMeshDescription& MeshDescription,
 
 		// We have a left-handed coordinate system, but a counter-clockwise winding order
 		// Hence normal calculation has to take the triangle vectors cross product in reverse.
-		FVector3f TmpNormal = FVector3f::CrossProduct(DPosition2, DPosition1).GetSafeNormal(AdjustedComparisonThreshold);
+		FVector3f TmpNormal = FVector3f::CrossProduct(DPosition2, DPosition1).GetSafeNormal(SquareComparisonThreshold);
 		if (!TmpNormal.IsNearlyZero(ComparisonThreshold))
 		{
 			FMatrix44f	ParameterToLocal(
@@ -182,11 +184,13 @@ void FStaticMeshOperations::ComputePolygonTangentsAndNormals(FMeshDescription& M
 
 static TTuple<FVector3f, FVector3f, FVector3f> GetTriangleTangentsAndNormals(float ComparisonThreshold, TArrayView<const FVector3f> VertexPositions, TArrayView<const FVector2D> VertexUVs)
 {
-	float AdjustedComparisonThreshold = FMath::Max(ComparisonThreshold, MIN_flt);
+	// GetSafeNormal compare the squareSum to the tolerance.
+	const float SquareComparisonThreshold = FMath::Max(ComparisonThreshold * ComparisonThreshold, MIN_flt);
 
 	const FVector3f Position0 = VertexPositions[0];
-	const FVector3f DPosition1 = VertexPositions[1] - Position0;
-	const FVector3f DPosition2 = VertexPositions[2] - Position0;
+	// To avoid numerical error due to small edges, DPosition is normalized
+	const FVector3f DPosition1 = (VertexPositions[1] - Position0).GetSafeNormal(SquareComparisonThreshold);
+	const FVector3f DPosition2 = (VertexPositions[2] - Position0).GetSafeNormal(SquareComparisonThreshold);
 
 	const FVector2f UV0 = FVector2f(VertexUVs[0]);
 	const FVector2f DUV1 = FVector2f(VertexUVs[1]) - UV0;
@@ -194,7 +198,7 @@ static TTuple<FVector3f, FVector3f, FVector3f> GetTriangleTangentsAndNormals(flo
 
 	// We have a left-handed coordinate system, but a counter-clockwise winding order
 	// Hence normal calculation has to take the triangle vectors cross product in reverse.
-	FVector3f Normal = FVector3f::CrossProduct(DPosition2, DPosition1).GetSafeNormal(AdjustedComparisonThreshold);
+	FVector3f Normal = FVector3f::CrossProduct(DPosition2, DPosition1).GetSafeNormal(SquareComparisonThreshold);
 	if (!Normal.IsNearlyZero(ComparisonThreshold))
 	{
 		FMatrix44f	ParameterToLocal(
