World Partition - HLOD: Prevent most common case of UV generation failure

diff --git a/Engine/Source/Runtime/StaticMeshDescription/Private/StaticMeshOperations.cpp b/Engine/Source/Runtime/StaticMeshDescription/Private/StaticMeshOperations.cpp
--- a/Engine/Source/Runtime/StaticMeshDescription/Private/StaticMeshOperations.cpp
+++ b/Engine/Source/Runtime/StaticMeshDescription/Private/StaticMeshOperations.cpp
@@ -2087,7 +2087,13 @@ bool FStaticMeshOperations::GenerateUniqueUVsForStaticMesh(const FMeshDescriptio
 	FLayoutUV Packer(MeshDescriptionView);
 	int32 NumCharts = Packer.FindCharts(OverlappingCorners);
 
-	bool bPackSuccess = Packer.FindBestPacking(FMath::Clamp(TextureResolution / 4, 32, 512));
+	// Scale down texture resolution to speed up UV generation time
+	// Packing expects at least one texel per chart. This is the absolute minimum to generate valid UVs.
+	const int32 PackingResolution = FMath::Clamp(TextureResolution / 4, 32, 512);
+	const int32 AbsoluteMinResolution = 1 << FMath::CeilLogTwo(FMath::Sqrt((float)NumCharts));
+	const int32 FinalPackingResolution = FMath::Max(TextureResolution, AbsoluteMinResolution);
+
+	bool bPackSuccess = Packer.FindBestPacking(FinalPackingResolution);
 	if (bPackSuccess)
 	{
 		Packer.CommitPackedUVs();
