World Partition - HLOD: Fix copy paste bug prevent the use of a lower resolution for UV packing in FStaticMeshOperations::GenerateUniqueUVsForStaticMesh()

diff --git a/Engine/Source/Runtime/StaticMeshDescription/Private/StaticMeshOperations.cpp b/Engine/Source/Runtime/StaticMeshDescription/Private/StaticMeshOperations.cpp
--- a/Engine/Source/Runtime/StaticMeshDescription/Private/StaticMeshOperations.cpp
+++ b/Engine/Source/Runtime/StaticMeshDescription/Private/StaticMeshOperations.cpp
@@ -2091,7 +2091,7 @@ bool FStaticMeshOperations::GenerateUniqueUVsForStaticMesh(const FMeshDescriptio
 	// Packing expects at least one texel per chart. This is the absolute minimum to generate valid UVs.
 	const int32 PackingResolution = FMath::Clamp(TextureResolution / 4, 32, 512);
 	const int32 AbsoluteMinResolution = 1 << FMath::CeilLogTwo(FMath::Sqrt((float)NumCharts));
-	const int32 FinalPackingResolution = FMath::Max(TextureResolution, AbsoluteMinResolution);
+	const int32 FinalPackingResolution = FMath::Max(PackingResolution, AbsoluteMinResolution);
 
 	bool bPackSuccess = Packer.FindBestPacking(FinalPackingResolution);
 	if (bPackSuccess)
