World Partition - HLOD: Improvements to FUniqueUVMeshDescriptionView
* Reduce artifacts from wrong input by filtering out degenerate triangles

diff --git a/Engine/Source/Runtime/StaticMeshDescription/Private/StaticMeshOperations.cpp b/Engine/Source/Runtime/StaticMeshDescription/Private/StaticMeshOperations.cpp
--- a/Engine/Source/Runtime/StaticMeshDescription/Private/StaticMeshOperations.cpp
+++ b/Engine/Source/Runtime/StaticMeshDescription/Private/StaticMeshOperations.cpp
@@ -1934,15 +1934,30 @@ struct FUniqueUVMeshDescriptionView final : FLayoutUV::IMeshView
 	TVertexInstanceAttributesConstRef<FVector3f> Normals;
 	TVertexInstanceAttributesConstRef<FVector2f> TexCoords;
 
-	uint32 NumIndices = 0;
-
-	TArray<int32>						RemapVerts;
-	TMap<FVertexInstanceID, int32>		InstanceToIndex;
+	TArray<FTriangleID>					RemapTriangles;
 	TArray<FVertexInstanceID>			UniqueVerts;
-	TArray<FVector2f>					OutputTexCoords;
+	TArray<FVector2D>&					OutputTexCoords;
+	bool								bMustRemap;
+
+	bool IsDegenerateTriangle(const FVector2f& A, const FVector2f& B, const FVector2f& C)
+	{
+		const FVector2f AB = B - A;
+		const FVector2f AC = C - A;
+		const float DoubleArea = FMath::Abs(AB ^ AC);
+		return DoubleArea < 2 * UE_KINDA_SMALL_NUMBER;
+	}
+
+	bool IsDegenerateTriangle(const FVector3f& A, const FVector3f& B, const FVector3f& C)
+	{
+		const FVector3f AB = B - A;
+		const FVector3f AC = C - A;
+		const float DoubleAreaSquared = (AB ^ AC).SizeSquared();
+		return DoubleAreaSquared < 2 * UE_KINDA_SMALL_NUMBER * UE_KINDA_SMALL_NUMBER;
+	}
 
-	FUniqueUVMeshDescriptionView(const FMeshDescription& InMeshDescription, bool bMergeIdenticalMaterials)
+	FUniqueUVMeshDescriptionView(const FMeshDescription& InMeshDescription, bool bMergeIdenticalMaterials, TArray<FVector2D>& InOutTexCoords)
 		: MeshDescription(InMeshDescription)
+		, OutputTexCoords(InOutTexCoords)
 	{
 		FStaticMeshConstAttributes Attributes(MeshDescription);
 		Positions = Attributes.GetVertexPositions();
@@ -1951,17 +1966,17 @@ struct FUniqueUVMeshDescriptionView final : FLayoutUV::IMeshView
 
 		TVertexInstanceAttributesConstRef<FVector4f> VertexColors = Attributes.GetVertexInstanceColors();
 
-		NumIndices = MeshDescription.VertexInstances().Num();
+		int32 NumVertexInstances = MeshDescription.VertexInstances().Num();
+		int32 NumTriangles = MeshDescription.Triangles().Num();
 
-		OutputTexCoords.SetNumUninitialized(NumIndices);
-		UniqueVerts.Reserve(NumIndices);
+		OutputTexCoords.SetNumZeroed(NumVertexInstances);
+		UniqueVerts.Reserve(NumVertexInstances);
+		RemapTriangles.Reserve(NumTriangles);
 
 		TMap<uint32, FTriangleID> UniqueTriangles;
-
 		if (bMergeIdenticalMaterials)
 		{
-			RemapVerts.Reserve(NumIndices);
-			UniqueTriangles.Reserve(MeshDescription.Triangles().Num());
+			UniqueTriangles.Reserve(NumTriangles);
 		}
 
 		auto HashAttribute = [](FVertexInstanceID InVertexInstanceID, auto InAttributeArrayRef, int32& TriangleHash)
@@ -1978,52 +1993,53 @@ struct FUniqueUVMeshDescriptionView final : FLayoutUV::IMeshView
 		for (const FTriangleID TriangleID : MeshDescription.Triangles().GetElementIDs())
 		{
 			const FPolygonGroupID RefPolygonGroupID = MeshDescription.GetTrianglePolygonGroup(TriangleID);
-			TArrayView<const FVertexInstanceID> VertexInstances = MeshDescription.GetTriangleVertexInstances(TriangleID);
+			TConstArrayView<const FVertexInstanceID> VertexInstancesIDs = MeshDescription.GetTriangleVertexInstances(TriangleID);
+			TConstArrayView<const FVertexID> VertexIDs = MeshDescription.GetTriangleVertices(TriangleID);
 
-			bool bUnique = true;
-			if (bMergeIdenticalMaterials)
-			{
-				int32 TriangleHash = GetTypeHash(RefPolygonGroupID);
-				for (const FVertexInstanceID& VertexInstanceID : VertexInstances)
-				{
-					// Compute hash based on UVs & vertices colors
-					HashAttribute(VertexInstanceID, TexCoords, TriangleHash);
-					HashAttribute(VertexInstanceID, VertexColors, TriangleHash);
-				}
+			FTriangleID RemapTriangleID = TriangleID;
 
-				FTriangleID& UniqueTriangle = UniqueTriangles.FindOrAdd(TriangleHash);
-				if (UniqueTriangle != INDEX_NONE)
+			// Filter out degenerate triangles
+			bool bDegenerate = false;
+			bDegenerate = bDegenerate || IsDegenerateTriangle(Positions[VertexIDs[0]], Positions[VertexIDs[1]], Positions[VertexIDs[2]]);
+			bDegenerate = bDegenerate || IsDegenerateTriangle(TexCoords[VertexInstancesIDs[0]], TexCoords[VertexInstancesIDs[1]], TexCoords[VertexInstancesIDs[2]]);
+			if (!bDegenerate)
+			{
+				bool bUnique = true;
+				if (bMergeIdenticalMaterials)
 				{
-					for (const FVertexInstanceID& UniqueVertexInstance : MeshDescription.GetTriangleVertexInstances(UniqueTriangle))
+					int32 TriangleHash = GetTypeHash(RefPolygonGroupID);
+					for (const FVertexInstanceID& VertexInstanceID : VertexInstancesIDs)
 					{
-						RemapVerts.Add(InstanceToIndex[UniqueVertexInstance]);
+						// Compute hash based on UVs & vertices colors
+						HashAttribute(VertexInstanceID, TexCoords, TriangleHash);
+						HashAttribute(VertexInstanceID, VertexColors, TriangleHash);
 					}
 
-					bUnique = false;
-				}
-				else
-				{
-					UniqueTriangle = TriangleID;
-				}
-			}
-			
-			if (bUnique)
-			{
-				for (const FVertexInstanceID& VertexInstanceID : VertexInstances)
-				{
-					if (bMergeIdenticalMaterials)
+					FTriangleID* UniqueTriangleIDPtr = UniqueTriangles.Find(TriangleHash);
+					if (UniqueTriangleIDPtr != nullptr)
 					{
-						InstanceToIndex.Emplace(VertexInstanceID, UniqueVerts.Num());
-						RemapVerts.Add(UniqueVerts.Num());
+						RemapTriangleID = *UniqueTriangleIDPtr;
+						bUnique = false;
 					}
+					else
+					{
+						UniqueTriangles.Add(TriangleHash, TriangleID);
+					}
+				}
 
-					UniqueVerts.Add(VertexInstanceID);
+				if (bUnique)
+				{
+					UniqueVerts.Append(VertexInstancesIDs);
 				}
 			}
+
+			RemapTriangles.Add(RemapTriangleID);
 		}
+
+		bMustRemap = bMergeIdenticalMaterials && UniqueTriangles.Num() != RemapTriangles.Num();
 	}
 
-	uint32 GetNumIndices() const override 
+	uint32 GetNumIndices() const override
 	{
 		return UniqueVerts.Num();
 	}
@@ -2051,25 +2067,31 @@ struct FUniqueUVMeshDescriptionView final : FLayoutUV::IMeshView
 
 	void SetOutputTexcoord(uint32 Index, const FVector2f& Value) override
 	{
-		OutputTexCoords[Index] = Value;
+		OutputTexCoords[UniqueVerts[Index]] = FVector2D(Value);
 	}
 
-	TArray<FVector2f> RetrievePackedUVs()
+	void ResolvePackedUVs()
 	{
-		if (RemapVerts.IsEmpty())
-		{ 
-			return OutputTexCoords;
-		}
-		
-		TArray<FVector2f> Out;
-		Out.Reserve(NumIndices);
-
-		for (uint32 i = 0; i < NumIndices; ++i)
+		if (bMustRemap)
 		{
-			Out.Add(OutputTexCoords[RemapVerts[i]]);
+		    int32 RemapTrianglesIdx = 0;
+		    for (const FTriangleID TriangleID : MeshDescription.Triangles().GetElementIDs())
+		    {
+			    FTriangleID RemapTriangleID = RemapTriangles[RemapTrianglesIdx];
+			    if (RemapTriangleID != TriangleID)
+			    {
+				    TConstArrayView<const FVertexInstanceID> SrcVertexInstancesIDs = MeshDescription.GetTriangleVertexInstances(RemapTriangleID);
+				    TConstArrayView<const FVertexInstanceID> DstVertexInstancesIDs = MeshDescription.GetTriangleVertexInstances(TriangleID);
+    
+				    for (int32 i = 0; i < 3; i++)
+				    {
+					    OutputTexCoords[DstVertexInstancesIDs[i]] = OutputTexCoords[SrcVertexInstancesIDs[i]];
+				    }
+			    }
+    
+			    RemapTrianglesIdx++;
+		    }
 		}
-
-		return Out;
 	}
 };
 
@@ -2077,7 +2099,7 @@ bool FStaticMeshOperations::GenerateUniqueUVsForStaticMesh(const FMeshDescriptio
 {
 	TRACE_CPUPROFILER_EVENT_SCOPE(FStaticMeshOperations::GenerateUniqueUVsForStaticMesh)
 
-	FUniqueUVMeshDescriptionView MeshDescriptionView(MeshDescription, bMergeIdenticalMaterials);
+	FUniqueUVMeshDescriptionView MeshDescriptionView(MeshDescription, bMergeIdenticalMaterials, OutTexCoords);
 
 	// Find overlapping corners for UV generator. Allow some threshold - this should not produce any error in a case if resulting
 	// mesh will not merge these vertices.
@@ -2097,7 +2119,7 @@ bool FStaticMeshOperations::GenerateUniqueUVsForStaticMesh(const FMeshDescriptio
 	if (bPackSuccess)
 	{
 		Packer.CommitPackedUVs();
-		OutTexCoords = UE::LWC::ConvertArrayType<FVector2D>(MeshDescriptionView.RetrievePackedUVs());	// LWC_TODO: Perf pessimization.
+		MeshDescriptionView.ResolvePackedUVs();
 	}
 
 	return bPackSuccess;
