World Partition - HLOD: Use lower threshold to detect degenerate triangles in order to avoid false positive (UE_SMALL_NUMBER vs UE_KINDA_SMALL_NUMBER)

diff --git a/Engine/Source/Runtime/StaticMeshDescription/Private/StaticMeshOperations.cpp b/Engine/Source/Runtime/StaticMeshDescription/Private/StaticMeshOperations.cpp
--- a/Engine/Source/Runtime/StaticMeshDescription/Private/StaticMeshOperations.cpp
+++ b/Engine/Source/Runtime/StaticMeshDescription/Private/StaticMeshOperations.cpp
@@ -1944,7 +1944,7 @@ struct FUniqueUVMeshDescriptionView final : FLayoutUV::IMeshView
 		const FVector2f AB = B - A;
 		const FVector2f AC = C - A;
 		const float DoubleArea = FMath::Abs(AB ^ AC);
-		return DoubleArea < 2 * UE_KINDA_SMALL_NUMBER;
+		return DoubleArea < 2 * UE_SMALL_NUMBER;
 	}
 
 	bool IsDegenerateTriangle(const FVector3f& A, const FVector3f& B, const FVector3f& C)
@@ -1952,7 +1952,7 @@ struct FUniqueUVMeshDescriptionView final : FLayoutUV::IMeshView
 		const FVector3f AB = B - A;
 		const FVector3f AC = C - A;
 		const float DoubleAreaSquared = (AB ^ AC).SizeSquared();
-		return DoubleAreaSquared < 2 * UE_KINDA_SMALL_NUMBER * UE_KINDA_SMALL_NUMBER;
+		return DoubleAreaSquared < 2 * UE_SMALL_NUMBER * UE_SMALL_NUMBER;
 	}
 
 	FUniqueUVMeshDescriptionView(const FMeshDescription& InMeshDescription, bool bMergeIdenticalMaterials, TArray<FVector2D>& InOutTexCoords)
