Add Insights scope on GetUniqueVertexCount

#rnx

diff --git a/Engine/Source/Runtime/StaticMeshDescription/Private/StaticMeshOperations.cpp b/Engine/Source/Runtime/StaticMeshDescription/Private/StaticMeshOperations.cpp
--- a/Engine/Source/Runtime/StaticMeshDescription/Private/StaticMeshOperations.cpp
+++ b/Engine/Source/Runtime/StaticMeshDescription/Private/StaticMeshOperations.cpp
@@ -2616,6 +2616,8 @@ int32 FStaticMeshOperations::GetUniqueVertexCount(const FMeshDescription& MeshDe
 
 int32 FStaticMeshOperations::GetUniqueVertexCount(const FMeshDescription& MeshDescription, const FOverlappingCorners& OverlappingCorners)
 {
+	TRACE_CPUPROFILER_EVENT_SCOPE(FStaticMeshOperations::GetUniqueVertexCount);
+
 	constexpr uint32 NumTexCoords = MAX_MESH_TEXTURE_COORDS_MD;
 	TArray< UE::Private::FPrivateVertexInfo > Verts;
 	Verts.Reserve(MeshDescription.Vertices().Num());
