Warn when building a static mesh that has NaNs in it.

diff --git a/Engine/Source/Developer/MeshBuilder/Private/MeshDescriptionHelper.cpp b/Engine/Source/Developer/MeshBuilder/Private/MeshDescriptionHelper.cpp
--- a/Engine/Source/Developer/MeshBuilder/Private/MeshDescriptionHelper.cpp
+++ b/Engine/Source/Developer/MeshBuilder/Private/MeshDescriptionHelper.cpp
@@ -48,7 +48,7 @@ void FMeshDescriptionHelper::SetupRenderMeshDescription(UObject* Owner, FMeshDes
 	}
 
 	//This function make sure the Polygon Normals Tangents Binormals are computed and also remove degenerated triangle from the render mesh description.
-	FStaticMeshOperations::ComputeTriangleTangentsAndNormals(RenderMeshDescription, ComparisonThreshold);
+	FStaticMeshOperations::ComputeTriangleTangentsAndNormals(RenderMeshDescription, ComparisonThreshold, *Owner->GetPathName());
 
 	FVertexInstanceArray& VertexInstanceArray = RenderMeshDescription.VertexInstances();
 
diff --git a/Engine/Source/Runtime/StaticMeshDescription/Private/StaticMeshOperations.cpp b/Engine/Source/Runtime/StaticMeshDescription/Private/StaticMeshOperations.cpp
--- a/Engine/Source/Runtime/StaticMeshDescription/Private/StaticMeshOperations.cpp
+++ b/Engine/Source/Runtime/StaticMeshDescription/Private/StaticMeshOperations.cpp
@@ -189,6 +189,7 @@ static TTuple<FVector3f, FVector3f, FVector3f> GetTriangleTangentsAndNormals(flo
 
 	const FVector3f Position0 = VertexPositions[0];
 	// To avoid numerical error due to small edges, DPosition is normalized
+	// If the positions deltas are too small, we get a zero vector out.
 	const FVector3f DPosition1 = (VertexPositions[1] - Position0).GetSafeNormal(SquareComparisonThreshold);
 	const FVector3f DPosition2 = (VertexPositions[2] - Position0).GetSafeNormal(SquareComparisonThreshold);
 
@@ -198,6 +199,7 @@ static TTuple<FVector3f, FVector3f, FVector3f> GetTriangleTangentsAndNormals(flo
 
 	// We have a left-handed coordinate system, but a counter-clockwise winding order
 	// Hence normal calculation has to take the triangle vectors cross product in reverse.
+	// If we got a zero vector out above, then this is also zero
 	FVector3f Normal = FVector3f::CrossProduct(DPosition2, DPosition1).GetSafeNormal(SquareComparisonThreshold);
 	if (!Normal.IsNearlyZero(ComparisonThreshold))
 	{
@@ -244,7 +246,7 @@ static TTuple<FVector3f, FVector3f, FVector3f> GetTriangleTangentsAndNormals(flo
 }
 
 
-void FStaticMeshOperations::ComputeTriangleTangentsAndNormals(FMeshDescription& MeshDescription, float ComparisonThreshold)
+void FStaticMeshOperations::ComputeTriangleTangentsAndNormals(FMeshDescription& MeshDescription, float ComparisonThreshold, const TCHAR* DebugName)
 {
 	TRACE_CPUPROFILER_EVENT_SCOPE(FStaticMeshOperations::ComputeTriangleTangentsAndNormals_Selection);
 
@@ -264,7 +266,7 @@ void FStaticMeshOperations::ComputeTriangleTangentsAndNormals(FMeshDescription&
 	const int32 BatchCount = (NumTriangles + BatchSize - 1) / BatchSize;
 
 	ParallelFor( TEXT("ComputeTriangleTangentsAndNormals.PF"), BatchCount,1,
-		[BatchSize, ComparisonThreshold, NumTriangles, &Attributes](int32 BatchIndex)
+		[BatchSize, ComparisonThreshold, NumTriangles, &Attributes, DebugName](int32 BatchIndex)
 		{
 			TArrayView<const FVector3f> VertexPositions = Attributes.GetVertexPositions().GetRawArray();
 			TArrayView<const FVector2f> VertexUVs = Attributes.GetVertexInstanceUVs().GetRawArray();
@@ -293,6 +295,17 @@ void FStaticMeshOperations::ComputeTriangleTangentsAndNormals(FMeshDescription&
 					VertexPositions[TriangleVertexIDs[TriIndex + 2]]
 				};
 
+				if (TriangleVertexPositions[0].ContainsNaN() ||
+					TriangleVertexPositions[1].ContainsNaN() ||
+					TriangleVertexPositions[2].ContainsNaN())
+				{
+					UE_LOG(LogStaticMeshOperations, Warning, TEXT("Static Mesh %s has NaNs in it's vertex positions! Triangle index %d -- using identity for tangent basis."), DebugName ? DebugName : TEXT("<null>"), StartIndex);
+					TriangleNormals[StartIndex] = FVector3f(1, 0, 0);
+					TriangleTangents[StartIndex] = FVector3f(0, 1, 0);
+					TriangleBinormals[StartIndex] = FVector3f(0, 0, 1);
+					continue;
+				}
+
 				FVector2D TriangleUVs[3] =
 				{
 					FVector2D(VertexUVs[TriangleVertexInstanceIDs[TriIndex]]),
@@ -300,6 +313,17 @@ void FStaticMeshOperations::ComputeTriangleTangentsAndNormals(FMeshDescription&
 					FVector2D(VertexUVs[TriangleVertexInstanceIDs[TriIndex + 2]])
 				};
 
+				if (TriangleUVs[0].ContainsNaN() ||
+					TriangleUVs[1].ContainsNaN() ||
+					TriangleUVs[2].ContainsNaN())
+				{
+					UE_LOG(LogStaticMeshOperations, Warning, TEXT("Static Mesh %s has NaNs in it's vertex uvs! Triangle index %d -- using identity for tangent basis."), DebugName ? DebugName : TEXT("<null>"), StartIndex);
+					TriangleNormals[StartIndex] = FVector3f(1, 0, 0);
+					TriangleTangents[StartIndex] = FVector3f(0, 1, 0);
+					TriangleBinormals[StartIndex] = FVector3f(0, 0, 1);
+					continue;
+				}
+
 				TTuple<FVector3f, FVector3f, FVector3f> Result = GetTriangleTangentsAndNormals(ComparisonThreshold, TriangleVertexPositions, TriangleUVs);
 				TriangleNormals[StartIndex] = Result.Get<0>();
 				TriangleTangents[StartIndex] = Result.Get<1>();
@@ -788,7 +812,7 @@ bool IsTriangleDegenerated(const FRawMesh& SourceRawMesh, const TArray<FVertexID
 	return (VertexIDs[0] == VertexIDs[1] || VertexIDs[0] == VertexIDs[2] || VertexIDs[1] == VertexIDs[2]);
 }
 
-void FStaticMeshOperations::ConvertFromRawMesh(const FRawMesh& SourceRawMesh, FMeshDescription& DestinationMeshDescription, const TMap<int32, FName>& MaterialMap, bool bSkipNormalsAndTangents)
+void FStaticMeshOperations::ConvertFromRawMesh(const FRawMesh& SourceRawMesh, FMeshDescription& DestinationMeshDescription, const TMap<int32, FName>& MaterialMap, bool bSkipNormalsAndTangents, const TCHAR* DebugName)
 {
 	DestinationMeshDescription.Empty();
 
@@ -942,7 +966,7 @@ void FStaticMeshOperations::ConvertFromRawMesh(const FRawMesh& SourceRawMesh, FM
 	if (!bSkipNormalsAndTangents && (!bHasNormals || !bHasTangents))
 	{
 		//DestinationMeshDescription.ComputePolygonTangentsAndNormals(0.0f);
-		ComputeTriangleTangentsAndNormals(DestinationMeshDescription, 0.0f);
+		ComputeTriangleTangentsAndNormals(DestinationMeshDescription, 0.0f, DebugName);
 
 		//Create the missing normals and recompute the tangents with MikkTSpace.
 		EComputeNTBsFlags ComputeNTBsOptions = EComputeNTBsFlags::Tangents | EComputeNTBsFlags::UseMikkTSpace | EComputeNTBsFlags::BlendOverlappingNormals;
diff --git a/Engine/Source/Runtime/StaticMeshDescription/Public/StaticMeshOperations.h b/Engine/Source/Runtime/StaticMeshDescription/Public/StaticMeshOperations.h
--- a/Engine/Source/Runtime/StaticMeshDescription/Public/StaticMeshOperations.h
+++ b/Engine/Source/Runtime/StaticMeshDescription/Public/StaticMeshOperations.h
@@ -71,7 +71,7 @@ public:
 	static void ComputePolygonTangentsAndNormals(FMeshDescription& MeshDescription, float ComparisonThreshold = 0.0f);
 
 	/** Set the triangle tangent, normal, binormal and triangleCenter for all triangles in the mesh description. */
-	static void ComputeTriangleTangentsAndNormals(FMeshDescription& MeshDescription, float ComparisonThreshold = 0.0f);
+	static void ComputeTriangleTangentsAndNormals(FMeshDescription& MeshDescription, float ComparisonThreshold = 0.0f, const TCHAR* DebugName = nullptr);
 
 	/** 
 	 * Recompute any invalid normal, tangent or Bi-Normal for every vertex in the mesh description with the given options.
@@ -94,7 +94,7 @@ public:
 	static void ConvertToRawMesh(const FMeshDescription& SourceMeshDescription, FRawMesh& DestinationRawMesh, const TMap<FName, int32>& MaterialMap);
 
 	/** Convert old FRawMesh format to MeshDescription. */
-	static void ConvertFromRawMesh(const FRawMesh& SourceRawMesh, FMeshDescription& DestinationMeshDescription, const TMap<int32, FName>& MaterialMap, bool bSkipNormalsAndTangents = false);
+	static void ConvertFromRawMesh(const FRawMesh& SourceRawMesh, FMeshDescription& DestinationMeshDescription, const TMap<int32, FName>& MaterialMap, bool bSkipNormalsAndTangents = false, const TCHAR* DebugName = nullptr);
 
 	static void AppendMeshDescription(const FMeshDescription& SourceMesh, FMeshDescription& TargetMesh, FAppendSettings& AppendSettings);
 	static void AppendMeshDescriptions(const TArray<const FMeshDescription*>& SourceMeshes, FMeshDescription& TargetMesh, FAppendSettings& AppendSettings);
