Rollback cl 21484193.
In CL 21484193, DPosition1 and DPosition2 were normalized to avoid a null normal in case of very small triangle. This was done without taking care of the computation of the tangent and binormal.
Tangent and binormal have to be aligned on the UV map.
Normalizing DPosition1 and DPosition2 breaks the relationship between the dimensions and the UV of a triangle. The calculation of the tangent is therefore false.
The use of a square threshold is sufficient to avoid null normal (CL 21484193).

diff --git a/Engine/Source/Runtime/StaticMeshDescription/Private/StaticMeshOperations.cpp b/Engine/Source/Runtime/StaticMeshDescription/Private/StaticMeshOperations.cpp
--- a/Engine/Source/Runtime/StaticMeshDescription/Private/StaticMeshOperations.cpp
+++ b/Engine/Source/Runtime/StaticMeshDescription/Private/StaticMeshOperations.cpp
@@ -56,9 +56,8 @@ static bool GetPolygonTangentsAndNormals(FMeshDescription& MeshDescription,
 		const FVertexID VertexID2 = MeshDescription.GetVertexInstanceVertex(TriangleVertexInstances[2]);
 
 		const FVector3f Position0 = VertexPositions[VertexID0];
-		// To avoid numerical error due to small edges, DPosition is normalized
-		const FVector3f DPosition1 = (VertexPositions[VertexID1] - Position0).GetSafeNormal(SquareComparisonThreshold);
-		const FVector3f DPosition2 = (VertexPositions[VertexID2] - Position0).GetSafeNormal(SquareComparisonThreshold);
+		const FVector3f DPosition1 = VertexPositions[VertexID1] - Position0;
+		const FVector3f DPosition2 = VertexPositions[VertexID2] - Position0;
 
 		const FVector2f UV0 = VertexUVs[TriangleVertexInstances[0]];
 		const FVector2f DUV1 = VertexUVs[TriangleVertexInstances[1]] - UV0;
@@ -188,10 +187,9 @@ static TTuple<FVector3f, FVector3f, FVector3f> GetTriangleTangentsAndNormals(flo
 	const float SquareComparisonThreshold = FMath::Max(ComparisonThreshold * ComparisonThreshold, MIN_flt);
 
 	const FVector3f Position0 = VertexPositions[0];
-	// To avoid numerical error due to small edges, DPosition is normalized
 	// If the positions deltas are too small, we get a zero vector out.
-	const FVector3f DPosition1 = (VertexPositions[1] - Position0).GetSafeNormal(SquareComparisonThreshold);
-	const FVector3f DPosition2 = (VertexPositions[2] - Position0).GetSafeNormal(SquareComparisonThreshold);
+	const FVector3f DPosition1 = VertexPositions[1] - Position0;
+	const FVector3f DPosition2 = VertexPositions[2] - Position0;
 
 	const FVector2f UV0 = FVector2f(VertexUVs[0]);
 	const FVector2f DUV1 = FVector2f(VertexUVs[1]) - UV0;
