Interchange: Fix importing multi uv channel that was importing only the first UV channel.

#rnx

diff --git a/Engine/Plugins/Interchange/Runtime/Source/Import/Private/Mesh/InterchangeSkeletalMeshFactory.cpp b/Engine/Plugins/Interchange/Runtime/Source/Import/Private/Mesh/InterchangeSkeletalMeshFactory.cpp
--- a/Engine/Plugins/Interchange/Runtime/Source/Import/Private/Mesh/InterchangeSkeletalMeshFactory.cpp
+++ b/Engine/Plugins/Interchange/Runtime/Source/Import/Private/Mesh/InterchangeSkeletalMeshFactory.cpp
@@ -518,7 +518,8 @@ namespace UE
 						//Add the skinning in the mesh description
 						{
 							FSkeletalMeshAttributes PayloadSkeletalMeshAttributes(LodMeshPayload->LodMeshDescription);
-							PayloadSkeletalMeshAttributes.Register();
+							constexpr bool bKeepExistingAttribute = true;
+							PayloadSkeletalMeshAttributes.Register(bKeepExistingAttribute);
 							using namespace UE::AnimationCore;
 							TArray<FBoneWeight> BoneWeights;
 							FBoneWeight& BoneWeight = BoneWeights.AddDefaulted_GetRef();
diff --git a/Engine/Plugins/Interchange/Runtime/Source/Parsers/Fbx/Private/FbxMesh.cpp b/Engine/Plugins/Interchange/Runtime/Source/Parsers/Fbx/Private/FbxMesh.cpp
--- a/Engine/Plugins/Interchange/Runtime/Source/Parsers/Fbx/Private/FbxMesh.cpp
+++ b/Engine/Plugins/Interchange/Runtime/Source/Parsers/Fbx/Private/FbxMesh.cpp
@@ -295,7 +295,8 @@ bool FMeshDescriptionImporter::FillMeshDescriptionFromFbxMesh(FbxMesh* Mesh, TAr
 	}
 
 	FStaticMeshAttributes Attributes(*MeshDescription);
-	Attributes.Register();
+	constexpr bool bKeepExistingAttribute = true;
+	Attributes.Register(bKeepExistingAttribute);
 
 	bool bStaticMeshUseSmoothEdgesIfSmoothingInformationIsMissing = true;
 
@@ -1014,7 +1015,7 @@ bool FMeshDescriptionImporter::FillMeshDescriptionFromFbxMesh(FbxMesh* Mesh, TAr
 		if (MeshType == EMeshType::Skinned)
 		{
 			FSkeletalMeshAttributes SkeletalMeshAttributes(*MeshDescription);
-			SkeletalMeshAttributes.Register();
+			SkeletalMeshAttributes.Register(true);
 
 			using namespace UE::AnimationCore;
 			TMap<FVertexID, TArray<FBoneWeight>> RawBoneWeights;
diff --git a/Engine/Source/Runtime/MeshDescription/Private/MeshAttributes.cpp b/Engine/Source/Runtime/MeshDescription/Private/MeshAttributes.cpp
--- a/Engine/Source/Runtime/MeshDescription/Private/MeshAttributes.cpp
+++ b/Engine/Source/Runtime/MeshDescription/Private/MeshAttributes.cpp
@@ -41,7 +41,7 @@ namespace MeshAttribute
 }
 
 
-void FMeshAttributes::Register()
+void FMeshAttributes::Register(bool bKeepExistingAttribute)
 {
 	// Nothing to do here: Vertex positions are already registered by the FMeshDescription constructor
 }
diff --git a/Engine/Source/Runtime/MeshDescription/Public/MeshAttributes.h b/Engine/Source/Runtime/MeshDescription/Public/MeshAttributes.h
--- a/Engine/Source/Runtime/MeshDescription/Public/MeshAttributes.h
+++ b/Engine/Source/Runtime/MeshDescription/Public/MeshAttributes.h
@@ -64,7 +64,7 @@ public:
 
 	virtual ~FMeshAttributes() = default;
 
-	virtual void Register();
+	virtual void Register(bool bKeepExistingAttribute = false);
 
 	static bool IsReservedAttributeName(const FName InAttributeName)
 	{
diff --git a/Engine/Source/Runtime/SkeletalMeshDescription/Private/SkeletalMeshAttributes.cpp b/Engine/Source/Runtime/SkeletalMeshDescription/Private/SkeletalMeshAttributes.cpp
--- a/Engine/Source/Runtime/SkeletalMeshDescription/Private/SkeletalMeshAttributes.cpp
+++ b/Engine/Source/Runtime/SkeletalMeshDescription/Private/SkeletalMeshAttributes.cpp
@@ -42,23 +42,35 @@ FSkeletalMeshAttributesShared(InMeshDescription)
 	}
 }
 
-void FSkeletalMeshAttributes::Register()
+void FSkeletalMeshAttributes::Register(bool bKeepExistingAttribute)
 {
-	MeshDescription.VertexAttributes().RegisterAttribute<int32[]>(MeshAttribute::Vertex::SkinWeights, 1, 0, EMeshAttributeFlags::Mandatory);
+	if (!MeshDescription.VertexAttributes().HasAttribute(MeshAttribute::Vertex::SkinWeights) || !bKeepExistingAttribute)
+	{
+		MeshDescription.VertexAttributes().RegisterAttribute<int32[]>(MeshAttribute::Vertex::SkinWeights, 1, 0, EMeshAttributeFlags::Mandatory);
+	}
 
 	if (MeshDescription.GetElements().Contains(BonesElementName) == false)
 	{
 		BoneElements = MeshDescription.GetElements().Emplace(BonesElementName).Get();
 	}
 
-	BoneAttributes().RegisterAttribute<FName>(MeshAttribute::Bone::Name, 1, NAME_None, EMeshAttributeFlags::Mandatory);
+	if (!MeshDescription.BoneAttributes().HasAttribute(MeshAttribute::Bone::Name) || !bKeepExistingAttribute)
+	{
+		BoneAttributes().RegisterAttribute<FName>(MeshAttribute::Bone::Name, 1, NAME_None, EMeshAttributeFlags::Mandatory);
+	}
 	
-	BoneAttributes().RegisterAttribute<int32>(MeshAttribute::Bone::ParentIndex, 1, INDEX_NONE, EMeshAttributeFlags::Mandatory);
+	if (!MeshDescription.BoneAttributes().HasAttribute(MeshAttribute::Bone::ParentIndex) || !bKeepExistingAttribute)
+	{
+		BoneAttributes().RegisterAttribute<int32>(MeshAttribute::Bone::ParentIndex, 1, INDEX_NONE, EMeshAttributeFlags::Mandatory);
+	}
 	
-	BoneAttributes().RegisterAttribute<FTransform>(MeshAttribute::Bone::Pose, 1, FTransform::Identity, EMeshAttributeFlags::Mandatory);
+	if (!MeshDescription.BoneAttributes().HasAttribute(MeshAttribute::Bone::Pose) || !bKeepExistingAttribute)
+	{
+		BoneAttributes().RegisterAttribute<FTransform>(MeshAttribute::Bone::Pose, 1, FTransform::Identity, EMeshAttributeFlags::Mandatory);
+	}
 
 	// Call super class
-	FStaticMeshAttributes::Register();
+	FStaticMeshAttributes::Register(bKeepExistingAttribute);
 }
 
 void FSkeletalMeshAttributes::RegisterColorAttribute()
diff --git a/Engine/Source/Runtime/SkeletalMeshDescription/Private/SkeletalMeshOperations.cpp b/Engine/Source/Runtime/SkeletalMeshDescription/Private/SkeletalMeshOperations.cpp
--- a/Engine/Source/Runtime/SkeletalMeshDescription/Private/SkeletalMeshOperations.cpp
+++ b/Engine/Source/Runtime/SkeletalMeshDescription/Private/SkeletalMeshOperations.cpp
@@ -17,7 +17,8 @@ void FSkeletalMeshOperations::AppendSkinWeight(const FMeshDescription& SourceMes
 	FSkeletalMeshConstAttributes SourceSkeletalMeshAttributes(SourceMesh);
 	
 	FSkeletalMeshAttributes TargetSkeletalMeshAttributes(TargetMesh);
-	TargetSkeletalMeshAttributes.Register();
+	constexpr bool bKeepExistingAttribute = true;
+	TargetSkeletalMeshAttributes.Register(bKeepExistingAttribute);
 	
 	FSkinWeightsVertexAttributesConstRef SourceVertexSkinWeights = SourceSkeletalMeshAttributes.GetVertexSkinWeights();
 	FSkinWeightsVertexAttributesRef TargetVertexSkinWeights = TargetSkeletalMeshAttributes.GetVertexSkinWeights();
diff --git a/Engine/Source/Runtime/SkeletalMeshDescription/Public/SkeletalMeshAttributes.h b/Engine/Source/Runtime/SkeletalMeshDescription/Public/SkeletalMeshAttributes.h
--- a/Engine/Source/Runtime/SkeletalMeshDescription/Public/SkeletalMeshAttributes.h
+++ b/Engine/Source/Runtime/SkeletalMeshDescription/Public/SkeletalMeshAttributes.h
@@ -145,7 +145,7 @@ public:
 
 	explicit FSkeletalMeshAttributes(FMeshDescription& InMeshDescription);
 
-	virtual void Register() override;
+	virtual void Register(bool bKeepExistingAttribute = false) override;
 	
 	//
 	// Skin Weights Methods
diff --git a/Engine/Source/Runtime/StaticMeshDescription/Private/StaticMeshAttributes.cpp b/Engine/Source/Runtime/StaticMeshDescription/Private/StaticMeshAttributes.cpp
--- a/Engine/Source/Runtime/StaticMeshDescription/Private/StaticMeshAttributes.cpp
+++ b/Engine/Source/Runtime/StaticMeshDescription/Private/StaticMeshAttributes.cpp
@@ -35,27 +35,52 @@ PRAGMA_ENABLE_DEPRECATION_WARNINGS
 
 
 
-void FStaticMeshAttributes::Register()
+void FStaticMeshAttributes::Register(bool bKeepExistingAttribute)
 {
 	// Add basic vertex attributes
 
 	// Add basic vertex instance attributes
-	MeshDescription.VertexInstanceAttributes().RegisterAttribute<FVector2f>(MeshAttribute::VertexInstance::TextureCoordinate, 1, FVector2f::ZeroVector, EMeshAttributeFlags::Lerpable | EMeshAttributeFlags::Mandatory);
-	MeshDescription.VertexInstanceAttributes().RegisterAttribute<FVector3f>(MeshAttribute::VertexInstance::Normal, 1, FVector3f::ZeroVector, EMeshAttributeFlags::AutoGenerated | EMeshAttributeFlags::Mandatory);
-	MeshDescription.VertexInstanceAttributes().RegisterAttribute<FVector3f>(MeshAttribute::VertexInstance::Tangent, 1, FVector3f::ZeroVector, EMeshAttributeFlags::AutoGenerated | EMeshAttributeFlags::Mandatory);
-	MeshDescription.VertexInstanceAttributes().RegisterAttribute<float>(MeshAttribute::VertexInstance::BinormalSign, 1, 0.0f, EMeshAttributeFlags::AutoGenerated | EMeshAttributeFlags::Mandatory);
-	MeshDescription.VertexInstanceAttributes().RegisterAttribute<FVector4f>(MeshAttribute::VertexInstance::Color, 1, FVector4f(1.0f, 1.0f, 1.0f, 1.0f), EMeshAttributeFlags::Lerpable | EMeshAttributeFlags::Mandatory);
+	if (!MeshDescription.VertexInstanceAttributes().HasAttribute(MeshAttribute::VertexInstance::TextureCoordinate) || !bKeepExistingAttribute)
+	{
+		MeshDescription.VertexInstanceAttributes().RegisterAttribute<FVector2f>(MeshAttribute::VertexInstance::TextureCoordinate, 1, FVector2f::ZeroVector, EMeshAttributeFlags::Lerpable | EMeshAttributeFlags::Mandatory);
+	}
+
+	if (!MeshDescription.VertexInstanceAttributes().HasAttribute(MeshAttribute::VertexInstance::Normal) || !bKeepExistingAttribute)
+	{
+		MeshDescription.VertexInstanceAttributes().RegisterAttribute<FVector3f>(MeshAttribute::VertexInstance::Normal, 1, FVector3f::ZeroVector, EMeshAttributeFlags::AutoGenerated | EMeshAttributeFlags::Mandatory);
+	}
+
+	if (!MeshDescription.VertexInstanceAttributes().HasAttribute(MeshAttribute::VertexInstance::Tangent) || !bKeepExistingAttribute)
+	{
+		MeshDescription.VertexInstanceAttributes().RegisterAttribute<FVector3f>(MeshAttribute::VertexInstance::Tangent, 1, FVector3f::ZeroVector, EMeshAttributeFlags::AutoGenerated | EMeshAttributeFlags::Mandatory);
+	}
+
+	if (!MeshDescription.VertexInstanceAttributes().HasAttribute(MeshAttribute::VertexInstance::BinormalSign) || !bKeepExistingAttribute)
+	{
+		MeshDescription.VertexInstanceAttributes().RegisterAttribute<float>(MeshAttribute::VertexInstance::BinormalSign, 1, 0.0f, EMeshAttributeFlags::AutoGenerated | EMeshAttributeFlags::Mandatory);
+	}
+
+	if (!MeshDescription.VertexInstanceAttributes().HasAttribute(MeshAttribute::VertexInstance::Color) || !bKeepExistingAttribute)
+	{
+		MeshDescription.VertexInstanceAttributes().RegisterAttribute<FVector4f>(MeshAttribute::VertexInstance::Color, 1, FVector4f(1.0f, 1.0f, 1.0f, 1.0f), EMeshAttributeFlags::Lerpable | EMeshAttributeFlags::Mandatory);
+	}
 
 	// Add basic edge attributes
-	MeshDescription.EdgeAttributes().RegisterAttribute<bool>(MeshAttribute::Edge::IsHard, 1, false, EMeshAttributeFlags::Mandatory);
+	if (!MeshDescription.EdgeAttributes().HasAttribute(MeshAttribute::Edge::IsHard) || !bKeepExistingAttribute)
+	{
+		MeshDescription.EdgeAttributes().RegisterAttribute<bool>(MeshAttribute::Edge::IsHard, 1, false, EMeshAttributeFlags::Mandatory);
+	}
 
 	// Add basic polygon attributes
 
 	// Add basic polygon group attributes
-	MeshDescription.PolygonGroupAttributes().RegisterAttribute<FName>(MeshAttribute::PolygonGroup::ImportedMaterialSlotName, 1, NAME_None, EMeshAttributeFlags::Mandatory); //The unique key to match the mesh material slot
+	if (!MeshDescription.PolygonGroupAttributes().HasAttribute(MeshAttribute::PolygonGroup::ImportedMaterialSlotName) || !bKeepExistingAttribute)
+	{
+		MeshDescription.PolygonGroupAttributes().RegisterAttribute<FName>(MeshAttribute::PolygonGroup::ImportedMaterialSlotName, 1, NAME_None, EMeshAttributeFlags::Mandatory); //The unique key to match the mesh material slot
+	}
 
 	// Call super class
-	FMeshAttributes::Register();
+	FMeshAttributes::Register(bKeepExistingAttribute);
 }
 
 
diff --git a/Engine/Source/Runtime/StaticMeshDescription/Private/StaticMeshOperations.cpp b/Engine/Source/Runtime/StaticMeshDescription/Private/StaticMeshOperations.cpp
--- a/Engine/Source/Runtime/StaticMeshDescription/Private/StaticMeshOperations.cpp
+++ b/Engine/Source/Runtime/StaticMeshDescription/Private/StaticMeshOperations.cpp
@@ -999,7 +999,8 @@ void FStaticMeshOperations::AppendMeshDescriptions(const TArray<const FMeshDescr
 	int32 NumEdges = 0;
 	int32 NumTriangles = 0;
 
-	int32 MaxNumUVChannels = TargetVertexInstanceUVs.GetNumChannels();
+	int32 MaxNumVertexInstanceUVChannels = TargetVertexInstanceUVs.GetNumChannels();
+	int32 MaxNumUVChannels = TargetMesh.GetNumUVElementChannels();
 	int32 MaxNumPolygonGroups = 0;
 	int32 MaxNumMeshVertices = 0;
 	int32 MaxNumEdges = 0;
@@ -1016,7 +1017,14 @@ void FStaticMeshOperations::AppendMeshDescriptions(const TArray<const FMeshDescr
 
 		FStaticMeshConstAttributes SourceAttributes(SourceMesh);
 		TVertexInstanceAttributesConstRef<FVector2f> SourceVertexInstanceUVs = SourceAttributes.GetVertexInstanceUVs();
-		for (int32 ChannelIdx = MaxNumUVChannels; ChannelIdx < SourceVertexInstanceUVs.GetNumChannels(); ++ChannelIdx)
+		for (int32 ChannelIdx = MaxNumVertexInstanceUVChannels; ChannelIdx < SourceVertexInstanceUVs.GetNumChannels(); ++ChannelIdx)
+		{
+			if (AppendSettings.bMergeUVChannels[ChannelIdx])
+			{
+				MaxNumVertexInstanceUVChannels = ChannelIdx + 1;
+			}
+		}
+		for (int32 ChannelIdx = MaxNumUVChannels; ChannelIdx < SourceMesh.GetNumUVElementChannels(); ++ChannelIdx)
 		{
 			if (AppendSettings.bMergeUVChannels[ChannelIdx])
 			{
@@ -1036,9 +1044,13 @@ void FStaticMeshOperations::AppendMeshDescriptions(const TArray<const FMeshDescr
 	TargetMesh.ReserveNewEdges(NumEdges);
 	TargetMesh.ReserveNewTriangles(NumTriangles);
 
-	if (MaxNumUVChannels > TargetVertexInstanceUVs.GetNumChannels())
+	if (MaxNumVertexInstanceUVChannels > TargetVertexInstanceUVs.GetNumChannels())
+	{
+		TargetVertexInstanceUVs.SetNumChannels(MaxNumVertexInstanceUVChannels);
+	}
+	if (MaxNumUVChannels > TargetMesh.GetNumUVElementChannels())
 	{
-		TargetVertexInstanceUVs.SetNumChannels(MaxNumUVChannels);
+		TargetMesh.SetNumUVChannels(MaxNumUVChannels);
 	}
 
 	PolygonGroupMap RemapPolygonGroup;
@@ -1064,6 +1076,21 @@ void FStaticMeshOperations::AppendMeshDescriptions(const TArray<const FMeshDescr
 		TVertexInstanceAttributesConstRef<FVector4f> SourceVertexInstanceColors = SourceAttributes.GetVertexInstanceColors();
 		TVertexInstanceAttributesConstRef<FVector2f> SourceVertexInstanceUVs = SourceAttributes.GetVertexInstanceUVs();
 
+		// Fill the UV arrays
+		const int32 NumUVChannel = FMath::Min(TargetMesh.GetNumUVElementChannels(), SourceMesh.GetNumUVElementChannels());
+		for (int32 UVLayerIndex = 0; UVLayerIndex < NumUVChannel; UVLayerIndex++)
+		{
+			TUVAttributesConstRef<FVector2f> SourceUVCoordinates = SourceMesh.UVAttributes(UVLayerIndex).GetAttributesRef<FVector2f>(MeshAttribute::UV::UVCoordinate);
+			TUVAttributesRef<FVector2f> TargetUVCoordinates = TargetMesh.UVAttributes(UVLayerIndex).GetAttributesRef<FVector2f>(MeshAttribute::UV::UVCoordinate);
+			int32 UVCount = SourceUVCoordinates.GetNumElements();
+			TargetMesh.ReserveNewUVs(UVCount, UVLayerIndex);
+			for (FUVID SourceUVID : SourceMesh.UVs(UVLayerIndex).GetElementIDs())
+			{
+				FUVID TargetUVID = TargetMesh.CreateUV(UVLayerIndex);
+				TargetUVCoordinates[TargetUVID] = SourceUVCoordinates[SourceUVID];
+			}
+		}
+
 		//PolygonGroups
 		if (AppendSettings.PolygonGroupsDelegate.IsBound())
 		{
@@ -1137,7 +1164,7 @@ void FStaticMeshOperations::AppendMeshDescriptions(const TArray<const FMeshDescr
 				TargetVertexInstanceColors[TargetVertexInstanceID] = SourceVertexInstanceColors[SourceVertexInstanceID];
 			}
 
-			for (int32 UVChannelIndex = 0; UVChannelIndex < MaxNumUVChannels && UVChannelIndex < SourceVertexInstanceUVs.GetNumChannels(); ++UVChannelIndex)
+			for (int32 UVChannelIndex = 0; UVChannelIndex < MaxNumVertexInstanceUVChannels && UVChannelIndex < SourceVertexInstanceUVs.GetNumChannels(); ++UVChannelIndex)
 			{
 				TargetVertexInstanceUVs.Set(TargetVertexInstanceID, UVChannelIndex, SourceVertexInstanceUVs.Get(SourceVertexInstanceID, UVChannelIndex));
 			}
diff --git a/Engine/Source/Runtime/StaticMeshDescription/Public/StaticMeshAttributes.h b/Engine/Source/Runtime/StaticMeshDescription/Public/StaticMeshAttributes.h
--- a/Engine/Source/Runtime/StaticMeshDescription/Public/StaticMeshAttributes.h
+++ b/Engine/Source/Runtime/StaticMeshDescription/Public/StaticMeshAttributes.h
@@ -77,7 +77,7 @@ public:
 		: FMeshAttributes(InMeshDescription)
 	{}
 
-	virtual void Register() override;
+	virtual void Register(bool bKeepExistingAttribute = false) override;
 	
 	static bool IsReservedAttributeName(const FName InAttributeName)
 	{
