Ensure the nanite hires meshdescription materials are in the same order has the LOD 0 meshdescription materials. We use the imported material name to reorder them properly.
- If there is some material that do not find a match, those materials will stay in the imported order. The import will output a log message saying the material names do not match.
- If there is less material in the hires mesh description, the import will create empty section in the hires mesh to ensure the order of the materials. No log message will be output in that case.
-If there is more material in the hires mesh description, the import will transfer extra material section triangle to the first valid hires material section data (i.e. this should be matching the LOD 0 first material section)

We do this reorder when we import a hires nanite mesh and also if we have already a hires mesh description and we re-import the LOD 0 base mesh. This ensure we always have the hires meshdescription materials in the same order of the LOD 0.

#rnx

diff --git a/Engine/Plugins/Interchange/Runtime/Source/Import/Private/Mesh/InterchangeStaticMeshFactory.cpp b/Engine/Plugins/Interchange/Runtime/Source/Import/Private/Mesh/InterchangeStaticMeshFactory.cpp
--- a/Engine/Plugins/Interchange/Runtime/Source/Import/Private/Mesh/InterchangeStaticMeshFactory.cpp
+++ b/Engine/Plugins/Interchange/Runtime/Source/Import/Private/Mesh/InterchangeStaticMeshFactory.cpp
@@ -640,6 +640,21 @@ UObject* UInterchangeStaticMeshFactory::EndImportAssetObject_GameThread(const FI
 		UInterchangeBaseNode::CopyStorage(StaticMeshFactoryNode, CurrentNode);
 		CurrentNode->FillAllCustomAttributeFromObject(StaticMesh);
 		UE::Interchange::FFactoryCommon::ApplyReimportStrategyToAsset(StaticMesh, PreviousNode, CurrentNode, StaticMeshFactoryNode);
+
+		//Reorder the Hires mesh description in the same order has the lod 0 mesh description
+		if (StaticMesh->IsHiResMeshDescriptionValid())
+		{
+			FMeshDescription* HiresMeshDescription = StaticMesh->GetHiResMeshDescription();
+			FMeshDescription* Lod0MeshDescription = StaticMesh->GetMeshDescription(0);
+			if (HiresMeshDescription && Lod0MeshDescription)
+			{
+				StaticMesh->ModifyHiResMeshDescription();
+				FString MaterialNameConflictMsg = TEXT("[Asset ") + StaticMesh->GetPathName() + TEXT("] Nanite hi - res import have some material name that differ from the LOD 0 material name.Your nanite hi - res should use the same material names the LOD 0 use to ensure we can remap the section in the same order.");
+				FString MaterialCountConflictMsg = TEXT("[Asset ") + StaticMesh->GetPathName() + TEXT("] Nanite hi-res import use more material then LOD 0. Your nanite hi-res should have less or equal number of material. Any extra material will be remap to the first material use by LOD 0.");
+				FStaticMeshOperations::ReorderMeshDescriptionPolygonGroups(*Lod0MeshDescription, *HiresMeshDescription, MaterialNameConflictMsg, MaterialCountConflictMsg);
+				StaticMesh->CommitHiResMeshDescription();
+			}
+		}
 	}
 #endif // WITH_EDITOR
 
diff --git a/Engine/Source/Editor/UnrealEd/Private/Factories/EditorFactories.cpp b/Engine/Source/Editor/UnrealEd/Private/Factories/EditorFactories.cpp
--- a/Engine/Source/Editor/UnrealEd/Private/Factories/EditorFactories.cpp
+++ b/Engine/Source/Editor/UnrealEd/Private/Factories/EditorFactories.cpp
@@ -276,6 +276,8 @@
 #include "TextureReferenceResolver.h"
 #include "UDIMUtilities.h"
 #include "FileHelpers.h"
+#include "StaticMeshOperations.h"
+
 
 DEFINE_LOG_CATEGORY(LogEditorFactories);
 
@@ -6160,6 +6162,20 @@ EReimportResult::Type UReimportFbxStaticMeshFactory::Reimport( UObject* Obj )
 				}
 
 				FFbxImporter->ImportStaticMeshGlobalSockets(Mesh);
+
+				if (Mesh->IsHiResMeshDescriptionValid())
+				{
+					FMeshDescription* HiresMeshDescription = Mesh->GetHiResMeshDescription();
+					FMeshDescription* Lod0MeshDescription = Mesh->GetMeshDescription(0);
+					if (HiresMeshDescription && Lod0MeshDescription)
+					{
+						Mesh->ModifyHiResMeshDescription();
+						FString MaterialNameConflictMsg = TEXT("[Asset ") + Mesh->GetPathName() + TEXT("] Nanite hi - res import have some material name that differ from the LOD 0 material name.Your nanite hi - res should use the same material names the LOD 0 use to ensure we can remap the section in the same order.");
+						FString MaterialCountConflictMsg = TEXT("[Asset ") + Mesh->GetPathName() + TEXT("] Nanite hi-res import use more material then LOD 0. Your nanite hi-res should have less or equal number of material. Any extra material will be remap to the first material use by LOD 0.");
+						FStaticMeshOperations::ReorderMeshDescriptionPolygonGroups(*Lod0MeshDescription, *HiresMeshDescription, MaterialNameConflictMsg, MaterialCountConflictMsg);
+						Mesh->CommitHiResMeshDescription();
+					}
+				}
 			}
 			else
 			{
diff --git a/Engine/Source/Editor/UnrealEd/Private/FbxMeshUtils.cpp b/Engine/Source/Editor/UnrealEd/Private/FbxMeshUtils.cpp
--- a/Engine/Source/Editor/UnrealEd/Private/FbxMeshUtils.cpp
+++ b/Engine/Source/Editor/UnrealEd/Private/FbxMeshUtils.cpp
@@ -40,6 +40,7 @@
 #include "ClothingAsset.h"
 #include "SkinWeightsUtilities.h"
 #include "LODUtilities.h"
+#include "StaticMeshOperations.h"
 
 DEFINE_LOG_CATEGORY_STATIC(LogExportMeshUtils, Log, All);
 
@@ -377,10 +378,17 @@ namespace FbxMeshUtils
 
 					BaseStaticMesh->ModifyHiResMeshDescription();
 
-					FMeshDescription* LOD0MeshDescription = TempStaticMesh->GetMeshDescription(0);
-					check(LOD0MeshDescription);
+					FMeshDescription* TempLOD0MeshDescription = TempStaticMesh->GetMeshDescription(0);
+					check(TempLOD0MeshDescription);
 
-					*HiResMeshDescription = MoveTemp(*LOD0MeshDescription);
+					if (FMeshDescription* BaseMeshDescription = BaseStaticMesh->GetMeshDescription(0))
+					{
+						FString MaterialNameConflictMsg = TEXT("[Asset ") + BaseStaticMesh->GetPathName() + TEXT("] Nanite hi - res import have some material name that differ from the LOD 0 material name.Your nanite hi - res should use the same material names the LOD 0 use to ensure we can remap the section in the same order.");
+						FString MaterialCountConflictMsg = TEXT("[Asset ") + BaseStaticMesh->GetPathName() + TEXT("] Nanite hi-res import use more material then LOD 0. Your nanite hi-res should have less or equal number of material. Any extra material will be remap to the first material use by LOD 0.");
+						FStaticMeshOperations::ReorderMeshDescriptionPolygonGroups(*BaseMeshDescription, *TempLOD0MeshDescription, MaterialNameConflictMsg, MaterialCountConflictMsg);
+					}
+					
+					*HiResMeshDescription = MoveTemp(*TempLOD0MeshDescription);
 
 					BaseStaticMesh->CommitHiResMeshDescription();
 
diff --git a/Engine/Source/Runtime/MeshDescription/Private/MeshDescription.cpp b/Engine/Source/Runtime/MeshDescription/Private/MeshDescription.cpp
--- a/Engine/Source/Runtime/MeshDescription/Private/MeshDescription.cpp
+++ b/Engine/Source/Runtime/MeshDescription/Private/MeshDescription.cpp
@@ -1986,6 +1986,33 @@ void FMeshDescription::RemapPolygonGroups(const TMap<FPolygonGroupID, FPolygonGr
 	}
 }
 
+void FMeshDescription::TransferPolygonGroup(FPolygonGroupID SourceID, FPolygonGroupID DestinationID)
+{
+	if (SourceID == DestinationID || !IsPolygonGroupValid(SourceID) || !IsPolygonGroupValid(DestinationID))
+	{
+		//Cannot transfer on self or if we have invalid polygon group ID
+		return;
+	}
+	TArray<FTriangleID> Triangles;
+	Triangles = PolygonGroupToTriangles.Find<FTriangleID>(SourceID);
+	TArray<FPolygonID> Polygons;
+	Polygons = PolygonGroupToPolygons.Find<FPolygonID>(SourceID);
+	PolygonGroupElements->Get().Remove(SourceID);
+	PolygonGroupToPolygons.RemoveKey(SourceID);
+	PolygonGroupToTriangles.RemoveKey(SourceID);
+
+	for (const FTriangleID TriangleID : Triangles)
+	{
+		TrianglePolygonGroups[TriangleID] = DestinationID;
+		PolygonGroupToTriangles.AddReferenceToKey(DestinationID, TriangleID);
+	}
+
+	for (const FPolygonID PolygonID : Polygons)
+	{
+		PolygonPolygonGroups[PolygonID] = DestinationID;
+		PolygonGroupToPolygons.AddReferenceToKey(DestinationID, PolygonID);
+	}
+}
 
 #if WITH_EDITORONLY_DATA
 
diff --git a/Engine/Source/Runtime/MeshDescription/Public/MeshDescription.h b/Engine/Source/Runtime/MeshDescription/Public/MeshDescription.h
--- a/Engine/Source/Runtime/MeshDescription/Public/MeshDescription.h
+++ b/Engine/Source/Runtime/MeshDescription/Public/MeshDescription.h
@@ -1194,7 +1194,9 @@ public:
 
 	/** Remaps polygon groups according to the supplied map */
 	void RemapPolygonGroups(const TMap<FPolygonGroupID, FPolygonGroupID>& Remap);
-
+	
+	/** Transfer the source polygon group data to the destination polygon group. The source data is append to the destination */
+	void TransferPolygonGroup(FPolygonGroupID SourceID, FPolygonGroupID DestinationID);
 
 	//////////////////////////////////////////////////////////////////////
 	// Whole mesh operations
diff --git a/Engine/Source/Runtime/StaticMeshDescription/Private/StaticMeshOperations.cpp b/Engine/Source/Runtime/StaticMeshDescription/Private/StaticMeshOperations.cpp
--- a/Engine/Source/Runtime/StaticMeshDescription/Private/StaticMeshOperations.cpp
+++ b/Engine/Source/Runtime/StaticMeshDescription/Private/StaticMeshOperations.cpp
@@ -2776,4 +2776,114 @@ int32 FStaticMeshOperations::GetUniqueVertexCount(const FMeshDescription& MeshDe
 	return Verts.Num();
 }
 
+void FStaticMeshOperations::ReorderMeshDescriptionPolygonGroups(const FMeshDescription& SourceMeshDescription
+	, FMeshDescription& DestinationMeshDescription
+	, TOptional<const FString> DestinationUnmatchMaterialName_Msg
+	, TOptional<const FString> DestinationHaveMorePolygonGroup_Msg)
+{
+	if (SourceMeshDescription.IsEmpty() || DestinationMeshDescription.IsEmpty() || DestinationMeshDescription.PolygonGroups().Num() <= 1)
+	{
+		//Nothing to re-order
+		return;
+	}
+	FStaticMeshConstAttributes SourceAttribute(SourceMeshDescription);
+	TPolygonGroupAttributesConstRef<FName> SourceMaterialSlotNameAttribute = SourceAttribute.GetPolygonGroupMaterialSlotNames();
+	FStaticMeshAttributes DestinationAttribute(DestinationMeshDescription);
+	TPolygonGroupAttributesConstRef<FName> DestinationMaterialSlotNameAttribute = DestinationAttribute.GetPolygonGroupMaterialSlotNames();
+	TMap<FPolygonGroupID, FPolygonGroupID> MatchPolygonGroupsDestSource;
+	TMap<FPolygonGroupID, bool> MaterialMatched;
+	MaterialMatched.Reserve(SourceMaterialSlotNameAttribute.GetNumElements());
+	for (FPolygonGroupID SourcePolygonGroupID : SourceMeshDescription.PolygonGroups().GetElementIDs())
+	{
+		MaterialMatched.Add(SourcePolygonGroupID, false);
+	}
+	for (FPolygonGroupID DestinationPolygonGroupID : DestinationMeshDescription.PolygonGroups().GetElementIDs())
+	{
+		const FName MaterialNameToMatch = DestinationMaterialSlotNameAttribute[DestinationPolygonGroupID];
+		FPolygonGroupID MatchPolygonGroupID = INDEX_NONE;
+		for (FPolygonGroupID SourcePolygonGroupID : SourceMeshDescription.PolygonGroups().GetElementIDs())
+		{
+			bool& bMaterialMatched = MaterialMatched.FindChecked(SourcePolygonGroupID);
+			if (!bMaterialMatched)
+			{
+				if (SourceMaterialSlotNameAttribute[SourcePolygonGroupID] == MaterialNameToMatch)
+				{
+					MatchPolygonGroupID = SourcePolygonGroupID;
+					bMaterialMatched = true;
+					break;
+				}
+			}
+		}
+		if (MatchPolygonGroupID != INDEX_NONE)
+		{
+			MatchPolygonGroupsDestSource.FindOrAdd(DestinationPolygonGroupID) = MatchPolygonGroupID;
+		}
+	}
+
+	bool bHaveLogUnmatchMaterial = false;
+	//Assign any unmatch polygon group to first available source polygon group
+	for (FPolygonGroupID SourcePolygonGroupID : SourceMeshDescription.PolygonGroups().GetElementIDs())
+	{
+		bool& bMaterialMatched = MaterialMatched.FindChecked(SourcePolygonGroupID);
+		if (!bMaterialMatched)
+		{
+			for (TPair<FPolygonGroupID, FPolygonGroupID>& MatchPolygonGroup : MatchPolygonGroupsDestSource)
+			{
+				if (MatchPolygonGroup.Value == INDEX_NONE)
+				{
+					if (!bHaveLogUnmatchMaterial && DestinationUnmatchMaterialName_Msg.IsSet())
+					{
+						UE_LOG(LogStaticMeshOperations, Warning, TEXT("%s"), *DestinationUnmatchMaterialName_Msg.GetValue());
+						bHaveLogUnmatchMaterial = true;
+					}
+					MatchPolygonGroup.Value = SourcePolygonGroupID;
+					bMaterialMatched = true;
+					break;
+				}
+			}
+			if (!bMaterialMatched)
+			{
+				//Insert a mesh description PolygonGroup that have no polygon (empty section)
+				//This allow to fill hole in case the hires have less section then the LOD 0
+				FPolygonGroupID NewDestinationPolygonGroupID = DestinationMeshDescription.CreatePolygonGroup();
+				MatchPolygonGroupsDestSource.FindOrAdd(NewDestinationPolygonGroupID) = SourcePolygonGroupID;
+				bMaterialMatched = true;
+			}
+		}
+	}
+
+	TMap<FPolygonGroupID, FPolygonGroupID> RemapGroup;
+	for (TPair<FPolygonGroupID, FPolygonGroupID>& MatchPolygonGroup : MatchPolygonGroupsDestSource)
+	{
+		check(MatchPolygonGroup.Value != INDEX_NONE);
+		RemapGroup.Add(MatchPolygonGroup.Key, MatchPolygonGroup.Value);
+	}
+	//Remap the polygon group with the correct ID
+	DestinationMeshDescription.RemapPolygonGroups(RemapGroup);
+
+	//Transfer any non match polygon group to the first valid polygon group (Hi res mesh use more section then LOD 0 mesh)
+	FPolygonGroupID ToDestinationID = DestinationMeshDescription.PolygonGroups().GetFirstValidID();
+	if (ToDestinationID != INDEX_NONE)
+	{
+		TArray< FPolygonGroupID > ToTransferPolygonGroups;
+		for (FPolygonGroupID DestinationPolygonGroupID : DestinationMeshDescription.PolygonGroups().GetElementIDs())
+		{
+			if (!MatchPolygonGroupsDestSource.Contains(DestinationPolygonGroupID))
+			{
+				ToTransferPolygonGroups.Add(DestinationPolygonGroupID);
+			}
+		}
+		if (ToTransferPolygonGroups.Num() > 0 && DestinationHaveMorePolygonGroup_Msg.IsSet())
+		{
+			UE_LOG(LogStaticMeshOperations, Warning, TEXT("%s"), *DestinationHaveMorePolygonGroup_Msg.GetValue());
+		}
+
+		for (FPolygonGroupID DestinationPolygonID : ToTransferPolygonGroups)
+		{
+			//TransferPolygonGroup take care of validating polygon group IDs and check they are not the same.
+			DestinationMeshDescription.TransferPolygonGroup(DestinationPolygonID, ToDestinationID);
+		}
+	}
+}
+
 #undef LOCTEXT_NAMESPACE
diff --git a/Engine/Source/Runtime/StaticMeshDescription/Public/StaticMeshOperations.h b/Engine/Source/Runtime/StaticMeshDescription/Public/StaticMeshOperations.h
--- a/Engine/Source/Runtime/StaticMeshDescription/Public/StaticMeshOperations.h
+++ b/Engine/Source/Runtime/StaticMeshDescription/Public/StaticMeshOperations.h
@@ -190,4 +190,13 @@ public:
 	 */
 	static int32 GetUniqueVertexCount(const FMeshDescription& MeshDescription);
 	static int32 GetUniqueVertexCount(const FMeshDescription& MeshDescription, const FOverlappingCorners& OverlappingCorners);
+
+	/**
+	 * Reorder the destination mesh description polygon groups in the same order has the source mesh description polygon groups.
+	 * This function use the imported material name to reorder the section.
+	 */
+	static void ReorderMeshDescriptionPolygonGroups(const FMeshDescription& SourceMeshDescription
+		, FMeshDescription& DestinationMeshDescription
+		, TOptional<const FString> UnmatchMaterialNameWarning
+		, TOptional<const FString> MoreDestinationPolygonGroupWarning);
 };
