Fix crash when importing a nanite hi res that have unmatched material names

#rnx

diff --git a/Engine/Plugins/Interchange/Runtime/Source/Import/Private/Mesh/InterchangeStaticMeshFactory.cpp b/Engine/Plugins/Interchange/Runtime/Source/Import/Private/Mesh/InterchangeStaticMeshFactory.cpp
--- a/Engine/Plugins/Interchange/Runtime/Source/Import/Private/Mesh/InterchangeStaticMeshFactory.cpp
+++ b/Engine/Plugins/Interchange/Runtime/Source/Import/Private/Mesh/InterchangeStaticMeshFactory.cpp
@@ -625,7 +625,7 @@ UInterchangeFactoryBase::FImportAssetResult UInterchangeStaticMeshFactory::EndIm
 			{
 				StaticMesh->ModifyHiResMeshDescription();
 				FString MaterialNameConflictMsg = TEXT("[Asset ") + StaticMesh->GetPathName() + TEXT("] Nanite hi - res import have some material name that differ from the LOD 0 material name.Your nanite hi - res should use the same material names the LOD 0 use to ensure we can remap the section in the same order.");
-				FString MaterialCountConflictMsg = TEXT("[Asset ") + StaticMesh->GetPathName() + TEXT("] Nanite hi-res import use more material then LOD 0. Your nanite hi-res should have less or equal number of material. Any extra material will be remap to the first material use by LOD 0.");
+				FString MaterialCountConflictMsg = TEXT("[Asset ") + StaticMesh->GetPathName() + TEXT("] Nanite hi-res import dont have the same material count then LOD 0. Your nanite hi-res should have equal number of material.");
 				FStaticMeshOperations::ReorderMeshDescriptionPolygonGroups(*Lod0MeshDescription, *HiresMeshDescription, MaterialNameConflictMsg, MaterialCountConflictMsg);
 				StaticMesh->CommitHiResMeshDescription();
 			}
diff --git a/Engine/Source/Editor/UnrealEd/Private/Factories/EditorFactories.cpp b/Engine/Source/Editor/UnrealEd/Private/Factories/EditorFactories.cpp
--- a/Engine/Source/Editor/UnrealEd/Private/Factories/EditorFactories.cpp
+++ b/Engine/Source/Editor/UnrealEd/Private/Factories/EditorFactories.cpp
@@ -6161,7 +6161,7 @@ EReimportResult::Type UReimportFbxStaticMeshFactory::Reimport( UObject* Obj )
 					{
 						Mesh->ModifyHiResMeshDescription();
 						FString MaterialNameConflictMsg = TEXT("[Asset ") + Mesh->GetPathName() + TEXT("] Nanite hi - res import have some material name that differ from the LOD 0 material name.Your nanite hi - res should use the same material names the LOD 0 use to ensure we can remap the section in the same order.");
-						FString MaterialCountConflictMsg = TEXT("[Asset ") + Mesh->GetPathName() + TEXT("] Nanite hi-res import use more material then LOD 0. Your nanite hi-res should have less or equal number of material. Any extra material will be remap to the first material use by LOD 0.");
+						FString MaterialCountConflictMsg = TEXT("[Asset ") + Mesh->GetPathName() + TEXT("] Nanite hi-res import dont have the same material count then LOD 0. Your nanite hi-res should have equal number of material.");
 						FStaticMeshOperations::ReorderMeshDescriptionPolygonGroups(*Lod0MeshDescription, *HiresMeshDescription, MaterialNameConflictMsg, MaterialCountConflictMsg);
 						Mesh->CommitHiResMeshDescription();
 					}
diff --git a/Engine/Source/Editor/UnrealEd/Private/FbxMeshUtils.cpp b/Engine/Source/Editor/UnrealEd/Private/FbxMeshUtils.cpp
--- a/Engine/Source/Editor/UnrealEd/Private/FbxMeshUtils.cpp
+++ b/Engine/Source/Editor/UnrealEd/Private/FbxMeshUtils.cpp
@@ -384,7 +384,7 @@ namespace FbxMeshUtils
 					if (FMeshDescription* BaseMeshDescription = BaseStaticMesh->GetMeshDescription(0))
 					{
 						FString MaterialNameConflictMsg = TEXT("[Asset ") + BaseStaticMesh->GetPathName() + TEXT("] Nanite hi - res import have some material name that differ from the LOD 0 material name.Your nanite hi - res should use the same material names the LOD 0 use to ensure we can remap the section in the same order.");
-						FString MaterialCountConflictMsg = TEXT("[Asset ") + BaseStaticMesh->GetPathName() + TEXT("] Nanite hi-res import use more material then LOD 0. Your nanite hi-res should have less or equal number of material. Any extra material will be remap to the first material use by LOD 0.");
+						FString MaterialCountConflictMsg = TEXT("[Asset ") + BaseStaticMesh->GetPathName() + TEXT("] Nanite hi-res import dont have the same material count then LOD 0. Your nanite hi-res should have equal number of material.");
 						FStaticMeshOperations::ReorderMeshDescriptionPolygonGroups(*BaseMeshDescription, *TempLOD0MeshDescription, MaterialNameConflictMsg, MaterialCountConflictMsg);
 					}
 					
diff --git a/Engine/Source/Runtime/StaticMeshDescription/Private/StaticMeshOperations.cpp b/Engine/Source/Runtime/StaticMeshDescription/Private/StaticMeshOperations.cpp
--- a/Engine/Source/Runtime/StaticMeshDescription/Private/StaticMeshOperations.cpp
+++ b/Engine/Source/Runtime/StaticMeshDescription/Private/StaticMeshOperations.cpp
@@ -2779,111 +2779,104 @@ int32 FStaticMeshOperations::GetUniqueVertexCount(const FMeshDescription& MeshDe
 void FStaticMeshOperations::ReorderMeshDescriptionPolygonGroups(const FMeshDescription& SourceMeshDescription
 	, FMeshDescription& DestinationMeshDescription
 	, TOptional<const FString> DestinationUnmatchMaterialName_Msg
-	, TOptional<const FString> DestinationHaveMorePolygonGroup_Msg)
+	, TOptional<const FString> DestinationPolygonGroupCountDifferFromSource_Msg)
 {
 	if (SourceMeshDescription.IsEmpty() || DestinationMeshDescription.IsEmpty() || DestinationMeshDescription.PolygonGroups().Num() <= 1)
 	{
 		//Nothing to re-order
 		return;
 	}
+
+	//Do not allow reorder if the material count is different between the destination and the source
+	if (DestinationMeshDescription.PolygonGroups().Num() != SourceMeshDescription.PolygonGroups().Num())
+	{
+		if(DestinationPolygonGroupCountDifferFromSource_Msg.IsSet())
+		{
+			UE_LOG(LogStaticMeshOperations, Warning, TEXT("%s"), *DestinationPolygonGroupCountDifferFromSource_Msg.GetValue());
+		}
+		return;
+	}
+
 	FStaticMeshConstAttributes SourceAttribute(SourceMeshDescription);
 	TPolygonGroupAttributesConstRef<FName> SourceMaterialSlotNameAttribute = SourceAttribute.GetPolygonGroupMaterialSlotNames();
 	FStaticMeshAttributes DestinationAttribute(DestinationMeshDescription);
 	TPolygonGroupAttributesConstRef<FName> DestinationMaterialSlotNameAttribute = DestinationAttribute.GetPolygonGroupMaterialSlotNames();
 	TMap<FPolygonGroupID, FPolygonGroupID> MatchPolygonGroupsDestSource;
-	TMap<FPolygonGroupID, bool> MaterialMatched;
-	MaterialMatched.Reserve(SourceMaterialSlotNameAttribute.GetNumElements());
+	TMap<FPolygonGroupID, bool> SourceMaterialMatched;
+	SourceMaterialMatched.Reserve(SourceMaterialSlotNameAttribute.GetNumElements());
 	for (FPolygonGroupID SourcePolygonGroupID : SourceMeshDescription.PolygonGroups().GetElementIDs())
 	{
-		MaterialMatched.Add(SourcePolygonGroupID, false);
+		SourceMaterialMatched.Add(SourcePolygonGroupID, false);
 	}
+	TMap<FPolygonGroupID, bool> DestinationMaterialMatched;
+	DestinationMaterialMatched.Reserve(DestinationMaterialSlotNameAttribute.GetNumElements());
 	for (FPolygonGroupID DestinationPolygonGroupID : DestinationMeshDescription.PolygonGroups().GetElementIDs())
 	{
-		const FName MaterialNameToMatch = DestinationMaterialSlotNameAttribute[DestinationPolygonGroupID];
-		FPolygonGroupID MatchPolygonGroupID = INDEX_NONE;
-		for (FPolygonGroupID SourcePolygonGroupID : SourceMeshDescription.PolygonGroups().GetElementIDs())
-		{
-			bool& bMaterialMatched = MaterialMatched.FindChecked(SourcePolygonGroupID);
-			if (!bMaterialMatched)
-			{
-				if (SourceMaterialSlotNameAttribute[SourcePolygonGroupID] == MaterialNameToMatch)
-				{
-					MatchPolygonGroupID = SourcePolygonGroupID;
-					bMaterialMatched = true;
-					break;
-				}
-			}
-		}
-		if (MatchPolygonGroupID != INDEX_NONE)
-		{
-			MatchPolygonGroupsDestSource.FindOrAdd(DestinationPolygonGroupID) = MatchPolygonGroupID;
-		}
+		DestinationMaterialMatched.Add(DestinationPolygonGroupID, false);
 	}
 
-	bool bHaveLogUnmatchMaterial = false;
-	//Assign any unmatch polygon group to first available source polygon group
-	for (FPolygonGroupID SourcePolygonGroupID : SourceMeshDescription.PolygonGroups().GetElementIDs())
+	//Find the material name match
+	for (TPair<FPolygonGroupID, bool>& DestinationMatched : DestinationMaterialMatched)
 	{
-		bool& bMaterialMatched = MaterialMatched.FindChecked(SourcePolygonGroupID);
-		if (!bMaterialMatched)
+		const FName MaterialNameToMatch = DestinationMaterialSlotNameAttribute[DestinationMatched.Key];
+		FPolygonGroupID MatchPolygonGroupID = INDEX_NONE;
+		for (TPair<FPolygonGroupID, bool>& SourceMatched : SourceMaterialMatched)
 		{
-			for (TPair<FPolygonGroupID, FPolygonGroupID>& MatchPolygonGroup : MatchPolygonGroupsDestSource)
+			if (SourceMatched.Value)
 			{
-				if (MatchPolygonGroup.Value == INDEX_NONE)
-				{
-					if (!bHaveLogUnmatchMaterial && DestinationUnmatchMaterialName_Msg.IsSet())
-					{
-						UE_LOG(LogStaticMeshOperations, Warning, TEXT("%s"), *DestinationUnmatchMaterialName_Msg.GetValue());
-						bHaveLogUnmatchMaterial = true;
-					}
-					MatchPolygonGroup.Value = SourcePolygonGroupID;
-					bMaterialMatched = true;
-					break;
-				}
+				continue;
 			}
-			if (!bMaterialMatched)
+			if (SourceMaterialSlotNameAttribute[SourceMatched.Key] == MaterialNameToMatch)
 			{
-				//Insert a mesh description PolygonGroup that have no polygon (empty section)
-				//This allow to fill hole in case the hires have less section then the LOD 0
-				FPolygonGroupID NewDestinationPolygonGroupID = DestinationMeshDescription.CreatePolygonGroup();
-				MatchPolygonGroupsDestSource.FindOrAdd(NewDestinationPolygonGroupID) = SourcePolygonGroupID;
-				bMaterialMatched = true;
+					
+				MatchPolygonGroupID = SourceMatched.Key;
+				SourceMatched.Value = true;
+				DestinationMatched.Value = true;
+				break;
 			}
 		}
+		if (MatchPolygonGroupID != INDEX_NONE)
+		{
+			MatchPolygonGroupsDestSource.FindOrAdd(DestinationMatched.Key) = MatchPolygonGroupID;
+		}
 	}
 
-	TMap<FPolygonGroupID, FPolygonGroupID> RemapGroup;
-	for (TPair<FPolygonGroupID, FPolygonGroupID>& MatchPolygonGroup : MatchPolygonGroupsDestSource)
+	if (MatchPolygonGroupsDestSource.Num() < DestinationMaterialMatched.Num() && DestinationUnmatchMaterialName_Msg.IsSet())
 	{
-		check(MatchPolygonGroup.Value != INDEX_NONE);
-		RemapGroup.Add(MatchPolygonGroup.Key, MatchPolygonGroup.Value);
+		UE_LOG(LogStaticMeshOperations, Warning, TEXT("%s"), *DestinationUnmatchMaterialName_Msg.GetValue());
 	}
-	//Remap the polygon group with the correct ID
-	DestinationMeshDescription.RemapPolygonGroups(RemapGroup);
 
-	//Transfer any non match polygon group to the first valid polygon group (Hi res mesh use more section then LOD 0 mesh)
-	FPolygonGroupID ToDestinationID = DestinationMeshDescription.PolygonGroups().GetFirstValidID();
-	if (ToDestinationID != INDEX_NONE)
+	//Iterate the unmatched destination and use the first unmatched source
+	for (TPair<FPolygonGroupID, bool>& DestinationMatched : DestinationMaterialMatched)
 	{
-		TArray< FPolygonGroupID > ToTransferPolygonGroups;
-		for (FPolygonGroupID DestinationPolygonGroupID : DestinationMeshDescription.PolygonGroups().GetElementIDs())
-		{
-			if (!MatchPolygonGroupsDestSource.Contains(DestinationPolygonGroupID))
-			{
-				ToTransferPolygonGroups.Add(DestinationPolygonGroupID);
-			}
-		}
-		if (ToTransferPolygonGroups.Num() > 0 && DestinationHaveMorePolygonGroup_Msg.IsSet())
+		if (DestinationMatched.Value)
 		{
-			UE_LOG(LogStaticMeshOperations, Warning, TEXT("%s"), *DestinationHaveMorePolygonGroup_Msg.GetValue());
+			//Skip this destination because its already matched
+			continue;
 		}
 
-		for (FPolygonGroupID DestinationPolygonID : ToTransferPolygonGroups)
+		//Match the first unmatched source we found
+		for (TPair<FPolygonGroupID, bool>& SourceMatched : SourceMaterialMatched)
 		{
-			//TransferPolygonGroup take care of validating polygon group IDs and check they are not the same.
-			DestinationMeshDescription.TransferPolygonGroup(DestinationPolygonID, ToDestinationID);
+			if (SourceMatched.Value)
+			{
+				//Skip this source because its already matched
+				continue;
+			}
+			//Force match
+			MatchPolygonGroupsDestSource.FindOrAdd(DestinationMatched.Key) = SourceMatched.Key;
+			DestinationMatched.Value = true;
+			SourceMatched.Value = true;
+			break;
 		}
 	}
+
+	//Since both source and destination have the same amount of material, the MatchPolygonGroupsDestSource should have the same count.
+	if(ensure(MatchPolygonGroupsDestSource.Num() == DestinationMaterialMatched.Num()))
+	{
+		//Remap the polygon group with the correct ID
+		DestinationMeshDescription.RemapPolygonGroups(MatchPolygonGroupsDestSource);
+	}
 }
 
 #undef LOCTEXT_NAMESPACE
diff --git a/Engine/Source/Runtime/StaticMeshDescription/Public/StaticMeshOperations.h b/Engine/Source/Runtime/StaticMeshDescription/Public/StaticMeshOperations.h
--- a/Engine/Source/Runtime/StaticMeshDescription/Public/StaticMeshOperations.h
+++ b/Engine/Source/Runtime/StaticMeshDescription/Public/StaticMeshOperations.h
@@ -194,9 +194,11 @@ public:
 	/**
 	 * Reorder the destination mesh description polygon groups in the same order has the source mesh description polygon groups.
 	 * This function use the imported material name to reorder the section.
+	 * 
+	 * Note: The material count MUST be the same for both the source and the destination. It will not reorder anything if the count is different.
 	 */
 	static void ReorderMeshDescriptionPolygonGroups(const FMeshDescription& SourceMeshDescription
 		, FMeshDescription& DestinationMeshDescription
 		, TOptional<const FString> UnmatchMaterialNameWarning
-		, TOptional<const FString> MoreDestinationPolygonGroupWarning);
+		, TOptional<const FString> DestinationPolygonGroupCountDifferFromSource_Msg);
 };
