FStaticMeshOperations::ApplyTransform fix for Normals and Tangents.

diff --git a/Engine/Source/Runtime/StaticMeshDescription/Private/StaticMeshOperations.cpp b/Engine/Source/Runtime/StaticMeshDescription/Private/StaticMeshOperations.cpp
--- a/Engine/Source/Runtime/StaticMeshDescription/Private/StaticMeshOperations.cpp
+++ b/Engine/Source/Runtime/StaticMeshDescription/Private/StaticMeshOperations.cpp
@@ -1172,21 +1172,24 @@ void FStaticMeshOperations::AppendMeshDescriptions(const TArray<const FMeshDescr
 			SourceToTargetVertexInstanceID.Add(SourceVertexInstanceID, TargetVertexInstanceID);
 		}
 
+		bool bReverseCulling = false;
 		// Transform vertex instances properties
 		if (AppendSettings.MeshTransform)
 		{
 			const FTransform& Transform = AppendSettings.MeshTransform.GetValue();
-			bool bFlipBinormal = Transform.GetDeterminant() < 0;
-			float BinormalSignsFactor = bFlipBinormal ? -1.f : 1.f;
+			const FMatrix TransformInverseTransposeMatrix = Transform.ToMatrixWithScale().Inverse().GetTransposed();
+
+			bReverseCulling = Transform.GetDeterminant() < 0;
+			float BinormalSignsFactor = bReverseCulling ? -1.f : 1.f;
 			for (const TPair<FVertexInstanceID, FVertexInstanceID>& VertexInstanceIDPair : SourceToTargetVertexInstanceID)
 			{
 				FVertexInstanceID InstanceID = VertexInstanceIDPair.Value;
 
 				FVector3f& Normal = TargetVertexInstanceNormals[InstanceID];
-				Normal = FVector3f(Transform.TransformVectorNoScale(FVector(Normal)));
+				Normal = (FVector3f)FVector(TransformInverseTransposeMatrix.TransformVector((FVector)Normal).GetSafeNormal());
 
 				FVector3f& Tangent = TargetVertexInstanceTangents[InstanceID];
-				Tangent = FVector3f(Transform.TransformVectorNoScale(FVector(Tangent)));
+				Tangent = (FVector3f)FVector(TransformInverseTransposeMatrix.TransformVector((FVector)Tangent).GetSafeNormal());
 
 				TargetVertexInstanceBinormalSigns[InstanceID] *= BinormalSignsFactor;
 			}
@@ -1202,10 +1205,21 @@ void FStaticMeshOperations::AppendMeshDescriptions(const TArray<const FMeshDescr
 
 			TArray<FVertexInstanceID, TInlineAllocator<3>> VertexInstanceIDs;
 			VertexInstanceIDs.Reserve(3);
-			for (const FVertexInstanceID& VertexInstanceID : TriangleVertexInstanceIDs)
+			if (bReverseCulling)
 			{
-				VertexInstanceIDs.Add(SourceToTargetVertexInstanceID[VertexInstanceID]);
+				for (int32 ReverseVertexInstanceIdIndex = TriangleVertexInstanceIDs.Num()-1; ReverseVertexInstanceIdIndex >= 0; ReverseVertexInstanceIdIndex--)
+				{
+					VertexInstanceIDs.Add(SourceToTargetVertexInstanceID[TriangleVertexInstanceIDs[ReverseVertexInstanceIdIndex]]);
+				}
 			}
+			else
+			{
+				for (const FVertexInstanceID& VertexInstanceID : TriangleVertexInstanceIDs)
+				{
+					VertexInstanceIDs.Add(SourceToTargetVertexInstanceID[VertexInstanceID]);
+				}
+			}
+			
 			// Insert a triangle into the mesh
 			TargetMesh.CreateTriangle(TargetPolygonGroupID, VertexInstanceIDs);
 		}
@@ -2589,8 +2603,7 @@ void FStaticMeshOperations::ApplyTransform(FMeshDescription& MeshDescription, co
 		VertexPositions[VertexID] = FVector4f(Transform.TransformPosition(FVector3d(VertexPositions[VertexID])));
 	}
 
-	FMatrix AdjointT = Transform.TransposeAdjoint();
-	AdjointT.RemoveScaling();
+	FMatrix TransformInverseTransposeMatrix = Transform.Inverse().GetTransposed();
 
 	const bool bIsMirrored = Transform.Determinant() < 0.f;
 	const float MulBy = bIsMirrored ? -1.f : 1.f;
@@ -2599,11 +2612,12 @@ void FStaticMeshOperations::ApplyTransform(FMeshDescription& MeshDescription, co
 	{
 		FVector3f Tangent = VertexInstanceTangents[VertexInstanceID];
 		FVector3f Normal = VertexInstanceNormals[VertexInstanceID];
-		float BinormalSign = VertexInstanceBinormalSigns[VertexInstanceID];
 
-		VertexInstanceTangents[VertexInstanceID] = (FVector3f)FVector(AdjointT.TransformVector((FVector)Tangent) * MulBy);
+		VertexInstanceTangents[VertexInstanceID] = (FVector3f)FVector(TransformInverseTransposeMatrix.TransformVector((FVector)Tangent).GetSafeNormal());
+		VertexInstanceNormals[VertexInstanceID] = (FVector3f)FVector(TransformInverseTransposeMatrix.TransformVector((FVector)Normal).GetSafeNormal());
+
+		float BinormalSign = VertexInstanceBinormalSigns[VertexInstanceID];
 		VertexInstanceBinormalSigns[VertexInstanceID] = BinormalSign * MulBy;
-		VertexInstanceNormals[VertexInstanceID] = (FVector3f)FVector(AdjointT.TransformVector((FVector)Normal) * MulBy);
 	}
 
 	if (bIsMirrored)
