FStaticMeshOperations::GenerateUniqueUVsForStaticMesh() - Fallback to AutoUV if UVs are invalid or missing

diff --git a/Engine/Source/Developer/MeshMergeUtilities/Private/MeshMergeModule.cpp b/Engine/Source/Developer/MeshMergeUtilities/Private/MeshMergeModule.cpp
--- a/Engine/Source/Developer/MeshMergeUtilities/Private/MeshMergeModule.cpp
+++ b/Engine/Source/Developer/MeshMergeUtilities/Private/MeshMergeModule.cpp
@@ -18,6 +18,10 @@ public:
 		return Utilities;
 	}
 
+	virtual void StartupModule() override
+	{
+		FModuleManager::Get().LoadModuleChecked("StaticMeshDescription");
+	}
 
 	virtual void ShutdownModule() override
 	{
diff --git a/Engine/Source/Runtime/StaticMeshDescription/Private/StaticMeshDescriptionModule.cpp b/Engine/Source/Runtime/StaticMeshDescription/Private/StaticMeshDescriptionModule.cpp
--- a/Engine/Source/Runtime/StaticMeshDescription/Private/StaticMeshDescriptionModule.cpp
+++ b/Engine/Source/Runtime/StaticMeshDescription/Private/StaticMeshDescriptionModule.cpp
@@ -13,7 +13,7 @@ public:
 
 	virtual void StartupModule() override
 	{
-		// Register any modular features here
+		FModuleManager::Get().LoadModule("GeometryProcessingInterfaces");
 	}
 
 	virtual void ShutdownModule() override
diff --git a/Engine/Source/Runtime/StaticMeshDescription/Private/StaticMeshOperations.cpp b/Engine/Source/Runtime/StaticMeshDescription/Private/StaticMeshOperations.cpp
--- a/Engine/Source/Runtime/StaticMeshDescription/Private/StaticMeshOperations.cpp
+++ b/Engine/Source/Runtime/StaticMeshDescription/Private/StaticMeshOperations.cpp
@@ -11,6 +11,12 @@
 #include "Misc/SecureHash.h"
 #include "OverlappingCorners.h"
 #include "RawMesh.h"
+#include "Modules/ModuleManager.h"
+
+#if WITH_EDITOR
+#include "IGeometryProcessingInterfacesModule.h"
+#include "GeometryProcessingInterfaces/MeshAutoUV.h"
+#endif
 
 #if WITH_MIKKTSPACE
 #include "mikktspace.h"
@@ -2140,30 +2146,74 @@ bool FStaticMeshOperations::GenerateUniqueUVsForStaticMesh(const FMeshDescriptio
 {
 	TRACE_CPUPROFILER_EVENT_SCOPE(FStaticMeshOperations::GenerateUniqueUVsForStaticMesh)
 
-	FUniqueUVMeshDescriptionView MeshDescriptionView(MeshDescription, bMergeIdenticalMaterials, OutTexCoords);
+	TVertexInstanceAttributesConstRef<FVector2f> VertexInstanceUVs = FStaticMeshConstAttributes(MeshDescription).GetVertexInstanceUVs();
 
-	// Find overlapping corners for UV generator. Allow some threshold - this should not produce any error in a case if resulting
-	// mesh will not merge these vertices.
-	FOverlappingCorners OverlappingCorners(MeshDescriptionView, THRESH_POINTS_ARE_SAME);
+	OutTexCoords.Reset();
 
-	// Generate new UVs
-	FLayoutUV Packer(MeshDescriptionView);
-	int32 NumCharts = Packer.FindCharts(OverlappingCorners);
+	const bool bHasUVs = VertexInstanceUVs.GetNumElements() > 0;
+	if (bHasUVs)
+	{
+		FUniqueUVMeshDescriptionView MeshDescriptionView(MeshDescription, bMergeIdenticalMaterials, OutTexCoords);
 
-	// Scale down texture resolution to speed up UV generation time
-	// Packing expects at least one texel per chart. This is the absolute minimum to generate valid UVs.
-	const int32 PackingResolution = FMath::Clamp(TextureResolution / 4, 32, 512);
-	const int32 AbsoluteMinResolution = 1 << FMath::CeilLogTwo(FMath::Sqrt((float)NumCharts));
-	const int32 FinalPackingResolution = FMath::Max(PackingResolution, AbsoluteMinResolution);
+		// Find overlapping corners for UV generator. Allow some threshold - this should not produce any error in a case if resulting
+		// mesh will not merge these vertices.
+		FOverlappingCorners OverlappingCorners(MeshDescriptionView, THRESH_POINTS_ARE_SAME);
 
-	bool bPackSuccess = Packer.FindBestPacking(FinalPackingResolution);
-	if (bPackSuccess)
+		// Generate new UVs
+		FLayoutUV Packer(MeshDescriptionView);
+		int32 NumCharts = Packer.FindCharts(OverlappingCorners);
+
+		// Scale down texture resolution to speed up UV generation time
+		// Packing expects at least one texel per chart. This is the absolute minimum to generate valid UVs.
+		const int32 PackingResolution = FMath::Clamp(TextureResolution / 4, 32, 512);
+		const int32 AbsoluteMinResolution = 1 << FMath::CeilLogTwo(FMath::Sqrt((float)NumCharts));
+		const int32 FinalPackingResolution = FMath::Max(PackingResolution, AbsoluteMinResolution);
+
+		bool bPackSuccess = Packer.FindBestPacking(FinalPackingResolution);
+		if (bPackSuccess)
+		{
+			Packer.CommitPackedUVs();
+			MeshDescriptionView.ResolvePackedUVs();
+		}
+		else
+		{
+			OutTexCoords.Reset();
+		}
+	}
+	
+#if WITH_EDITOR
+	// Missing/invalid UVs, use the AutoUV interface
+	if (OutTexCoords.IsEmpty())
 	{
-		Packer.CommitPackedUVs();
-		MeshDescriptionView.ResolvePackedUVs();
+		IGeometryProcessingInterfacesModule* GeomProcInterfaces = FModuleManager::Get().GetModulePtr<IGeometryProcessingInterfacesModule>("GeometryProcessingInterfaces");
+		if (GeomProcInterfaces)
+		{
+		    IGeometryProcessing_MeshAutoUV* MeshAutoUV = GeomProcInterfaces->GetMeshAutoUVImplementation();
+    
+		    FMeshDescription MeshCopy = MeshDescription;
+    
+		    IGeometryProcessing_MeshAutoUV::FOptions Options = MeshAutoUV->ConstructDefaultOptions();
+		    IGeometryProcessing_MeshAutoUV::FResults Results;
+		    MeshAutoUV->GenerateUVs(MeshCopy, Options, Results);
+    
+		    if (Results.ResultCode == IGeometryProcessing_MeshAutoUV::EResultCode::Success)
+		    {
+			    TVertexInstanceAttributesConstRef<FVector2f> TexCoords;
+			    FStaticMeshConstAttributes AttributesCopy(MeshCopy);
+			    TexCoords = AttributesCopy.GetVertexInstanceUVs();
+    
+			    for (const FVertexInstanceID VertexInstanceID : MeshCopy.VertexInstances().GetElementIDs())
+			    {
+				    OutTexCoords.Add(FVector2D(TexCoords.Get(VertexInstanceID, 0)));
+			    }
+		    }
+		}
 	}
+#endif
 
-	return bPackSuccess;
+	check(OutTexCoords.IsEmpty() || OutTexCoords.Num() == MeshDescription.VertexInstances().Num());
+
+	return !OutTexCoords.IsEmpty();
 }
 
 bool FStaticMeshOperations::AddUVChannel(FMeshDescription& MeshDescription)
diff --git a/Engine/Source/Runtime/StaticMeshDescription/StaticMeshDescription.Build.cs b/Engine/Source/Runtime/StaticMeshDescription/StaticMeshDescription.Build.cs
--- a/Engine/Source/Runtime/StaticMeshDescription/StaticMeshDescription.Build.cs
+++ b/Engine/Source/Runtime/StaticMeshDescription/StaticMeshDescription.Build.cs
@@ -23,6 +23,23 @@ namespace UnrealBuildTool.Rules
 				}
 			);
 
+			if (Target.bBuildEditor == true)
+			{
+				PrivateIncludePathModuleNames.AddRange(
+					new string[]
+					{
+						"GeometryProcessingInterfaces"
+					}
+				);
+
+				DynamicallyLoadedModuleNames.AddRange(
+					new string[]
+					{
+						"GeometryProcessingInterfaces"
+					}
+				);
+			}
+
 			AddEngineThirdPartyPrivateStaticDependencies(Target, "MikkTSpace");
 		}
 	}
