[Engine]

diff --git a/Engine/Source/Runtime/Sockets/Public/SocketSubsystem.h b/Engine/Source/Runtime/Sockets/Public/SocketSubsystem.h
--- a/Engine/Source/Runtime/Sockets/Public/SocketSubsystem.h
+++ b/Engine/Source/Runtime/Sockets/Public/SocketSubsystem.h
@@ -54,19 +54,19 @@ using FUniqueSocket = TUniquePtr<FSocket, FSocketDeleter>;
  * This is the base interface to abstract platform specific sockets API
  * differences.
  */
-class SOCKETS_API ISocketSubsystem
+class ISocketSubsystem
 {
 public:
 
 	/**
 	 * Get the singleton socket subsystem for the given named subsystem
 	 */
-	static ISocketSubsystem* Get(const FName& SubsystemName=NAME_None);
+	static SOCKETS_API ISocketSubsystem* Get(const FName& SubsystemName=NAME_None);
 
 	/**
 	 * Shutdown all registered subsystems
 	 */
-	static void ShutdownAllSystems();
+	static SOCKETS_API void ShutdownAllSystems();
 
 
 	virtual ~ISocketSubsystem() { }
@@ -136,7 +136,7 @@ public:
 	 *
 	 * @return the new socket or NULL if failed
 	 */
-	FUniqueSocket CreateUniqueSocket(const FName& SocketType, const FString& SocketDescription, bool bForceUDP = false);
+	SOCKETS_API FUniqueSocket CreateUniqueSocket(const FName& SocketType, const FString& SocketDescription, bool bForceUDP = false);
 
 	/**
 	 * Creates a socket using the given protocol name,  wrapped in a unique pointer that will call DestroySocket automatically - do not call it explicitly!
@@ -148,7 +148,7 @@ public:
 	 *
 	 * @return the new socket or NULL if failed
 	 */
-	FUniqueSocket CreateUniqueSocket(const FName& SocketType, const FString& SocketDescription, const FName& ProtocolName);
+	SOCKETS_API FUniqueSocket CreateUniqueSocket(const FName& SocketType, const FString& SocketDescription, const FName& ProtocolName);
 
 	/**
 	 * Creates a resolve info cached struct to hold the resolved address
@@ -157,7 +157,7 @@ public:
 	 *
 	 * @return the new resolved address or NULL if failed
 	 */
-	virtual class FResolveInfoCached* CreateResolveInfoCached(TSharedPtr<FInternetAddr> Addr) const;
+	SOCKETS_API virtual class FResolveInfoCached* CreateResolveInfoCached(TSharedPtr<FInternetAddr> Addr) const;
 
 	/**
 	 * Cleans up a socket class
@@ -238,7 +238,7 @@ public:
 	 *                   formatting results and can be safely left to the default value.
 	 *
 	 */
-	virtual void GetAddressInfoAsync(FAsyncGetAddressInfoCallback Callback, const TCHAR* HostName,
+	SOCKETS_API virtual void GetAddressInfoAsync(FAsyncGetAddressInfoCallback Callback, const TCHAR* HostName,
 		const TCHAR* ServiceName = nullptr, EAddressInfoFlags QueryFlags = EAddressInfoFlags::Default,
 		const FName ProtocolTypeName = NAME_None,
 		ESocketType SocketType = ESocketType::SOCKTYPE_Unknown);
@@ -283,7 +283,7 @@ public:
 	 *
 	 * @return the resolve info to query for the address
 	 */
-	virtual class FResolveInfo* GetHostByName(const ANSICHAR* HostName);
+	SOCKETS_API virtual class FResolveInfo* GetHostByName(const ANSICHAR* HostName);
 
 	/**
 	 * Some platforms require chat data (voice, text, etc.) to be placed into
@@ -345,7 +345,7 @@ public:
 	 * @param Flags					Flags for specifying how FRecvMulti should be initialized (for e.g. retrieving timestamps)
 	 * @return						Returns the platform specific FRecvMulti instance
 	 */
-	virtual TUniquePtr<FRecvMulti> CreateRecvMulti(int32 MaxNumPackets, int32 MaxPacketSize,
+	SOCKETS_API virtual TUniquePtr<FRecvMulti> CreateRecvMulti(int32 MaxNumPackets, int32 MaxPacketSize,
 													ERecvMultiFlags Flags=ERecvMultiFlags::None);
 
 	/**
@@ -377,7 +377,7 @@ public:
 	 *
 	 * @param Code the error code to check
 	 */
-	const TCHAR* GetSocketError(ESocketErrors Code = SE_GET_LAST_ERROR_CODE);
+	SOCKETS_API const TCHAR* GetSocketError(ESocketErrors Code = SE_GET_LAST_ERROR_CODE);
 
 	/**
 	 * Gets the list of addresses associated with the adapters on the local computer.
@@ -388,7 +388,7 @@ public:
 	 *
 	 * @return true on success, false otherwise.
 	 */
-	virtual bool GetLocalAdapterAddresses(TArray<TSharedPtr<FInternetAddr>>& OutAddresses);
+	SOCKETS_API virtual bool GetLocalAdapterAddresses(TArray<TSharedPtr<FInternetAddr>>& OutAddresses);
 
 	/**
 	 * Get a local IP to bind to.
@@ -396,7 +396,7 @@ public:
 	 * Typically, it is better to use GetLocalBindAddresses as it better supports hybrid network functionality
 	 * and less chances for connections to fail due to mismatched protocols.
 	 */
-	virtual TSharedRef<FInternetAddr> GetLocalBindAddr(FOutputDevice& Out);
+	SOCKETS_API virtual TSharedRef<FInternetAddr> GetLocalBindAddr(FOutputDevice& Out);
 
 	/**
 	 * Get bindable addresses that this machine can use as reported by GetAddressInfo with the BindableAddress flag.
@@ -405,7 +405,7 @@ public:
 	 *
 	 * @return If GetAddressInfo succeeded or multihome is specified, an array of addresses that can be binded on. Failure returns an empty array.
 	 */
-	virtual TArray<TSharedRef<FInternetAddr>> GetLocalBindAddresses();
+	SOCKETS_API virtual TArray<TSharedRef<FInternetAddr>> GetLocalBindAddresses();
 
 	/**
 	 * Bind to next available port.
@@ -417,7 +417,7 @@ public:
 	 *
 	 * @return The bound port number, or 0 on failure
 	 */
-	int32 BindNextPort(FSocket* Socket, FInternetAddr& Addr, int32 PortCount, int32 PortIncrement);
+	SOCKETS_API int32 BindNextPort(FSocket* Socket, FInternetAddr& Addr, int32 PortCount, int32 PortIncrement);
 
 	/**
 	 * Uses the platform specific look up to determine the host address
@@ -431,7 +431,7 @@ public:
 	 *
 	 * @return The local host address
 	 */
-	virtual TSharedRef<FInternetAddr> GetLocalHostAddr(FOutputDevice& Out, bool& bCanBindAll);
+	SOCKETS_API virtual TSharedRef<FInternetAddr> GetLocalHostAddr(FOutputDevice& Out, bool& bCanBindAll);
 
 	/**
 	 * Returns the multihome address if the flag is present and valid. For ease of use, 
@@ -441,7 +441,7 @@ public:
 	 *
 	 * @return If the multihome address was set and valid
 	 */
-	virtual bool GetMultihomeAddress(TSharedRef<FInternetAddr>& Addr);
+	SOCKETS_API virtual bool GetMultihomeAddress(TSharedRef<FInternetAddr>& Addr);
 
 	/**
 	 * Checks the host name cache for an existing entry (faster than resolving again)
@@ -451,7 +451,7 @@ public:
 	 *
 	 * @return true if the host was found, false otherwise
 	 */
-	bool GetHostByNameFromCache(const ANSICHAR* HostName, TSharedPtr<class FInternetAddr>& Addr);
+	SOCKETS_API bool GetHostByNameFromCache(const ANSICHAR* HostName, TSharedPtr<class FInternetAddr>& Addr);
 
 	/**
 	 * Stores the ip address with the matching host name
@@ -459,19 +459,19 @@ public:
 	 * @param HostName the host name to search for
 	 * @param Addr the IP that will be copied from
 	 */
-	void AddHostNameToCache(const ANSICHAR* HostName, TSharedPtr<class FInternetAddr> Addr);
+	SOCKETS_API void AddHostNameToCache(const ANSICHAR* HostName, TSharedPtr<class FInternetAddr> Addr);
 
 	/**
 	 * Removes the host name to ip mapping from the cache
 	 *
 	 * @param HostName the host name to search for
 	 */
-	void RemoveHostNameFromCache(const ANSICHAR* HostName);
+	SOCKETS_API void RemoveHostNameFromCache(const ANSICHAR* HostName);
 
 	/**
 	 * Returns true if FSocket::RecvMulti is supported by this socket subsystem
 	 */
-	virtual bool IsSocketRecvMultiSupported() const;
+	SOCKETS_API virtual bool IsSocketRecvMultiSupported() const;
 
 
 	/**
@@ -487,13 +487,13 @@ public:
 	 * @param Translation	The type of translation to perform on the timestamp (time delta is usually faster than local timestamp)
 	 * @return				Returns the translated timestamp or delta
 	 */
-	virtual double TranslatePacketTimestamp(const FPacketTimestamp& Timestamp,
+	SOCKETS_API virtual double TranslatePacketTimestamp(const FPacketTimestamp& Timestamp,
 											ETimestampTranslation Translation=ETimestampTranslation::LocalTimestamp);
 
 	/**
 	 * Returns true if FSocket::RecvFromWithPktInfo is supported by this socket subsystem.
 	 */
-	virtual bool IsRecvFromWithPktInfoSupported() const;
+	SOCKETS_API virtual bool IsRecvFromWithPktInfoSupported() const;
 
 protected:
 
@@ -501,8 +501,8 @@ protected:
 	 * Conversion functions from the SocketProtocolFamily enum to the new FName system.
 	 * For now, both are supported, but it's better to use the FName when possible.
 	 */
-	virtual ESocketProtocolFamily GetProtocolFamilyFromName(const FName& InProtocolName) const;
-	virtual FName GetProtocolNameFromFamily(ESocketProtocolFamily InProtocolFamily) const;
+	SOCKETS_API virtual ESocketProtocolFamily GetProtocolFamilyFromName(const FName& InProtocolName) const;
+	SOCKETS_API virtual FName GetProtocolNameFromFamily(ESocketProtocolFamily InProtocolFamily) const;
 
 private:
 
diff --git a/Engine/Source/Runtime/Sockets/Public/SocketTypes.h b/Engine/Source/Runtime/Sockets/Public/SocketTypes.h
--- a/Engine/Source/Runtime/Sockets/Public/SocketTypes.h
+++ b/Engine/Source/Runtime/Sockets/Public/SocketTypes.h
@@ -162,7 +162,7 @@ ENUM_CLASS_FLAGS(ERecvMultiFlags);
  * Stores the persistent state and packet buffers/data, for receiving packets with FSocket::RecvMulti.
  * To optimize performance, use only once instance of this struct, for the lifetime of the socket.
  */
-struct SOCKETS_API FRecvMulti : public FNoncopyable, public FVirtualDestructor
+struct FRecvMulti : public FNoncopyable, public FVirtualDestructor
 {
 	friend struct FUnixRecvMulti;
 	friend class FSocketUnix;
@@ -216,7 +216,7 @@ private:
 	 * @param InMaxPacketSize		The maximum supported packet size
 	 * @param bRetrieveTimestamps	Whether or not to support retrieving timestamps
 	 */
-	FRecvMulti(ISocketSubsystem* SocketSubsystem, int32 InMaxNumPackets, int32 InMaxPacketSize,
+	SOCKETS_API FRecvMulti(ISocketSubsystem* SocketSubsystem, int32 InMaxNumPackets, int32 InMaxPacketSize,
 				ERecvMultiFlags InitFlags=ERecvMultiFlags::None);
 
 
@@ -262,5 +262,5 @@ public:
 	 *
 	 * @param Ar	The archive being used to count the memory consumption
 	 */
-	virtual void CountBytes(FArchive& Ar) const;
+	SOCKETS_API virtual void CountBytes(FArchive& Ar) const;
 };
diff --git a/Engine/Source/Runtime/Sockets/Public/Sockets.h b/Engine/Source/Runtime/Sockets/Public/Sockets.h
--- a/Engine/Source/Runtime/Sockets/Public/Sockets.h
+++ b/Engine/Source/Runtime/Sockets/Public/Sockets.h
@@ -15,7 +15,7 @@ class FInternetAddr;
 /**
  * This is our abstract base class that hides the platform specific socket implementation
  */
-class SOCKETS_API FSocket
+class FSocket
 {
 protected:
 
@@ -179,7 +179,7 @@ public:
 	 * @param BytesSent Will indicate how much was sent.
 	 * @param Destination The network byte ordered address to send to.
 	 */
-	virtual bool SendTo(const uint8* Data, int32 Count, int32& BytesSent, const FInternetAddr& Destination);
+	SOCKETS_API virtual bool SendTo(const uint8* Data, int32 Count, int32& BytesSent, const FInternetAddr& Destination);
 
 	/**
 	 * Sends a buffer on a connected socket.
@@ -188,7 +188,7 @@ public:
 	 * @param Count The size of the data to send.
 	 * @param BytesSent Will indicate how much was sent.
 	 */
-	virtual bool Send(const uint8* Data, int32 Count, int32& BytesSent);
+	SOCKETS_API virtual bool Send(const uint8* Data, int32 Count, int32& BytesSent);
 
 	/**
 	 * Reads a chunk of data from the socket and gathers the source address.
@@ -204,7 +204,7 @@ public:
 	 * @param Flags The receive flags.
 	 * @return true on success, false in case of a closed socket or an unrecoverable error.
 	 */
-	virtual bool RecvFrom(uint8* Data, int32 BufferSize, int32& BytesRead, FInternetAddr& Source, ESocketReceiveFlags::Type Flags = ESocketReceiveFlags::None);
+	SOCKETS_API virtual bool RecvFrom(uint8* Data, int32 BufferSize, int32& BytesRead, FInternetAddr& Source, ESocketReceiveFlags::Type Flags = ESocketReceiveFlags::None);
 
 	/**
 	 * Reads a chunk of data from a connected socket
@@ -219,7 +219,7 @@ public:
 	 * @param Flags the receive flags
 	 * @return true on success, false in case of a closed socket or an unrecoverable error.
 	 */
-	virtual bool Recv(uint8* Data, int32 BufferSize, int32& BytesRead, ESocketReceiveFlags::Type Flags = ESocketReceiveFlags::None);
+	SOCKETS_API virtual bool Recv(uint8* Data, int32 BufferSize, int32& BytesRead, ESocketReceiveFlags::Type Flags = ESocketReceiveFlags::None);
 
 	/**
 	 * Reads multiple packets from the socket at once, gathering the source address and other optional platform specific data.
@@ -230,7 +230,7 @@ public:
 	 * @param Flags			The receive flags.
 	 * @return				Whether or not data was successfully received
 	 */
-	virtual bool RecvMulti(FRecvMulti& MultiData, ESocketReceiveFlags::Type Flags=ESocketReceiveFlags::None);
+	SOCKETS_API virtual bool RecvMulti(FRecvMulti& MultiData, ESocketReceiveFlags::Type Flags=ESocketReceiveFlags::None);
 
 	/**
 	 * Blocks until the specified condition is met.
@@ -415,7 +415,7 @@ public:
 	 * @param bRetrieveTimestamp	Whether to retrieve the timestamp upon receive
 	 * @return						True if the call succeeded, false otherwise.
 	 */
-	virtual bool SetRetrieveTimestamp(bool bRetrieveTimestamp=true);
+	SOCKETS_API virtual bool SetRetrieveTimestamp(bool bRetrieveTimestamp=true);
 
 	/**
 	 * Reads the port this socket is bound to.
@@ -431,7 +431,7 @@ public:
 	 * @param bEnable	Whether to enable/disable
 	 * @return True if the call succeeded, false otherwise.
 	 */
-	virtual bool SetIpPktInfo(bool bEnable);
+	SOCKETS_API virtual bool SetIpPktInfo(bool bEnable);
 
 	/**
 	 * Reads a chunk of data from the socket and gathers the source address and the destination using IP_PKTINFO
@@ -448,7 +448,7 @@ public:
 	 * @param Flags The receive flags.
 	 * @return true on success, false in case of a closed socket or an unrecoverable error.
 	 */
-	virtual bool RecvFromWithPktInfo(uint8* Data, int32 BufferSize, int32& BytesRead, FInternetAddr& Source, FInternetAddr& Destination, ESocketReceiveFlags::Type Flags = ESocketReceiveFlags::None);
+	SOCKETS_API virtual bool RecvFromWithPktInfo(uint8* Data, int32 BufferSize, int32& BytesRead, FInternetAddr& Source, FInternetAddr& Destination, ESocketReceiveFlags::Type Flags = ESocketReceiveFlags::None);
 
 	/**
 	 * Get the type of protocol the socket is bound to
diff --git a/Engine/Source/Runtime/SoundFieldRendering/Public/SoundFieldRendering.h b/Engine/Source/Runtime/SoundFieldRendering/Public/SoundFieldRendering.h
--- a/Engine/Source/Runtime/SoundFieldRendering/Public/SoundFieldRendering.h
+++ b/Engine/Source/Runtime/SoundFieldRendering/Public/SoundFieldRendering.h
@@ -12,7 +12,7 @@ SOUNDFIELDRENDERING_API TUniquePtr<ISoundfieldDecoderStream> CreateDefaultSource
 SOUNDFIELDRENDERING_API ISoundfieldEncodingSettingsProxy& GetAmbisonicsSourceDefaultSettings();
 
 
-class SOUNDFIELDRENDERING_API FAmbisonicsSoundfieldBuffer : public ISoundfieldAudioPacket
+class FAmbisonicsSoundfieldBuffer : public ISoundfieldAudioPacket
 {
 public:
 	// Interleaved audio buffer for all vector parts of the ambisonics stream.
@@ -36,48 +36,48 @@ public:
 	virtual ~FAmbisonicsSoundfieldBuffer()
 	{}
 
-	virtual void Serialize(FArchive& Ar) override;
-	virtual TUniquePtr<ISoundfieldAudioPacket> Duplicate() const override;
-	virtual void Reset() override;
+	SOUNDFIELDRENDERING_API virtual void Serialize(FArchive& Ar) override;
+	SOUNDFIELDRENDERING_API virtual TUniquePtr<ISoundfieldAudioPacket> Duplicate() const override;
+	SOUNDFIELDRENDERING_API virtual void Reset() override;
 
 };
 
-class SOUNDFIELDRENDERING_API FAmbisonicsSoundfieldSettings : public ISoundfieldEncodingSettingsProxy
+class FAmbisonicsSoundfieldSettings : public ISoundfieldEncodingSettingsProxy
 {
 public:
 	int32 Order;
 
-	virtual uint32 GetUniqueId() const override;
-	virtual TUniquePtr<ISoundfieldEncodingSettingsProxy> Duplicate() const override;
+	SOUNDFIELDRENDERING_API virtual uint32 GetUniqueId() const override;
+	SOUNDFIELDRENDERING_API virtual TUniquePtr<ISoundfieldEncodingSettingsProxy> Duplicate() const override;
 
 };
 
-class SOUNDFIELDRENDERING_API FSoundFieldDecoder
+class FSoundFieldDecoder
 {
 public:
-	FSoundFieldDecoder(); // initializes VirtualSpeakerWorldLockedSpeakerGains
+	SOUNDFIELDRENDERING_API FSoundFieldDecoder(); // initializes VirtualSpeakerWorldLockedSpeakerGains
 
-	void DecodeAudioDirectlyToDeviceOutputPositions(const FAmbisonicsSoundfieldBuffer& InputData, const FSoundfieldSpeakerPositionalData& OutputPositions, Audio::FAlignedFloatBuffer& OutputData);
-	void DecodeAudioToSevenOneAndDownmixToDevice(const FAmbisonicsSoundfieldBuffer& InputData, const FSoundfieldSpeakerPositionalData& OutputPositions, Audio::FAlignedFloatBuffer& OutputData);
-	static void RotateFirstOrderAmbisonicsBed(const FAmbisonicsSoundfieldBuffer& InputData, FAmbisonicsSoundfieldBuffer& OutputData, const FQuat& DestinationRotation, const FQuat& PreviousRotation);
+	SOUNDFIELDRENDERING_API void DecodeAudioDirectlyToDeviceOutputPositions(const FAmbisonicsSoundfieldBuffer& InputData, const FSoundfieldSpeakerPositionalData& OutputPositions, Audio::FAlignedFloatBuffer& OutputData);
+	SOUNDFIELDRENDERING_API void DecodeAudioToSevenOneAndDownmixToDevice(const FAmbisonicsSoundfieldBuffer& InputData, const FSoundfieldSpeakerPositionalData& OutputPositions, Audio::FAlignedFloatBuffer& OutputData);
+	static SOUNDFIELDRENDERING_API void RotateFirstOrderAmbisonicsBed(const FAmbisonicsSoundfieldBuffer& InputData, FAmbisonicsSoundfieldBuffer& OutputData, const FQuat& DestinationRotation, const FQuat& PreviousRotation);
 
-	static void FoaRotationInPlace(Audio::FAlignedFloatBuffer& InOutFrames, const float XRotDegrees, const float YRotDegrees, const float ZRotDegrees);
+	static SOUNDFIELDRENDERING_API void FoaRotationInPlace(Audio::FAlignedFloatBuffer& InOutFrames, const float XRotDegrees, const float YRotDegrees, const float ZRotDegrees);
 
 private:
 	// Special cased, vectorized versions of first order operations.
-	static void FirstOrderDecodeLoop(const int32 NumFrames, const int32 NumOutputChannels, const float* pAmbiFrame, const int32 NumAmbiChannels, float* SpeakerGainsPtr, float* OutputBufferPtrBuffer);
-	static void FirstOrderToSevenOneLoop(const int32 NumFrames, const int32 NumOutputChannels, const float* pAmbiFrame, const int32 NumAmbiChannels, float* SpeakerGainsPtr, float* OutputBufferPtrBuffer);
+	static SOUNDFIELDRENDERING_API void FirstOrderDecodeLoop(const int32 NumFrames, const int32 NumOutputChannels, const float* pAmbiFrame, const int32 NumAmbiChannels, float* SpeakerGainsPtr, float* OutputBufferPtrBuffer);
+	static SOUNDFIELDRENDERING_API void FirstOrderToSevenOneLoop(const int32 NumFrames, const int32 NumOutputChannels, const float* pAmbiFrame, const int32 NumAmbiChannels, float* SpeakerGainsPtr, float* OutputBufferPtrBuffer);
 
 	// Special cased vectorized versions of 3rd and fifth order operations.
-	static void OddOrderDecodeLoop(const int32 NumFrames, const int32 NumOutputChannels, const float* pAmbiFrame, const int32 NumAmbiChannels, float* SpeakerGainsPtr, float* OutputBufferPtrBuffer);
-	static void OddOrderToSevenOneLoop(const int32 NumFrames, const int32 NumOutputChannels, const float* pAmbiFrame, const int32 NumAmbiChannels, float* SpeakerGainsPtr, float* OutputBufferPtrBuffer);
+	static SOUNDFIELDRENDERING_API void OddOrderDecodeLoop(const int32 NumFrames, const int32 NumOutputChannels, const float* pAmbiFrame, const int32 NumAmbiChannels, float* SpeakerGainsPtr, float* OutputBufferPtrBuffer);
+	static SOUNDFIELDRENDERING_API void OddOrderToSevenOneLoop(const int32 NumFrames, const int32 NumOutputChannels, const float* pAmbiFrame, const int32 NumAmbiChannels, float* SpeakerGainsPtr, float* OutputBufferPtrBuffer);
 
 	// Partially vectorized decode operations for 2nd and 4th order operations.
-	static void EvenOrderDecodeLoop(const int32 NumFrames, const int32 NumOutputChannels, const float* pAmbiFrame, const int32 NumAmbiChannels, float* SpeakerGainsPtr, float* OutputBufferPtrBuffer);
-	static void EvenOrderToSevenOneLoop(const int32 NumFrames, const int32 NumOutputChannels, const float* pAmbiFrame, const int32 NumAmbiChannels, float* SpeakerGainsPtr, float* OutputBufferPtrBuffer);
+	static SOUNDFIELDRENDERING_API void EvenOrderDecodeLoop(const int32 NumFrames, const int32 NumOutputChannels, const float* pAmbiFrame, const int32 NumAmbiChannels, float* SpeakerGainsPtr, float* OutputBufferPtrBuffer);
+	static SOUNDFIELDRENDERING_API void EvenOrderToSevenOneLoop(const int32 NumFrames, const int32 NumOutputChannels, const float* pAmbiFrame, const int32 NumAmbiChannels, float* SpeakerGainsPtr, float* OutputBufferPtrBuffer);
 
 	// Helper function to get the default channel positions for a channel array.
-	static TArray<Audio::FChannelPositionInfo>* GetDefaultChannelPositions(int32 InNumChannels);
+	static SOUNDFIELDRENDERING_API TArray<Audio::FChannelPositionInfo>* GetDefaultChannelPositions(int32 InNumChannels);
 
 	Audio::FAlignedFloatBuffer VirtualSpeakerScratchBuffers;
 	Audio::FAlignedFloatBuffer FoaVirtualSpeakerWordLockedGains;
@@ -87,20 +87,20 @@ private:
 	Audio::FBufferLinearEase SpeakerGainLerper;
 	FVector2D LastListenerRotationSphericalCoord{ -999.f, -999.f }; // force interpolation on first callback
 
-	static FSphericalHarmonicCalculator SphereHarmCalc;
-	static TArray<Audio::FChannelPositionInfo> VirtualSpeakerLocationsHorzOnly;
-	static const VectorRegister Sqrt2Over2Vec;
-	static const VectorRegister ZeroVec;
+	static SOUNDFIELDRENDERING_API FSphericalHarmonicCalculator SphereHarmCalc;
+	static SOUNDFIELDRENDERING_API TArray<Audio::FChannelPositionInfo> VirtualSpeakerLocationsHorzOnly;
+	static SOUNDFIELDRENDERING_API const VectorRegister Sqrt2Over2Vec;
+	static SOUNDFIELDRENDERING_API const VectorRegister ZeroVec;
 
 	
 	friend class FSoundFieldEncoder;
 };
 
-class SOUNDFIELDRENDERING_API FSoundFieldEncoder
+class FSoundFieldEncoder
 {
 public:
-	FSoundFieldEncoder();
-	void EncodeAudioDirectlyFromOutputPositions(const Audio::FAlignedFloatBuffer& InputData, const FSoundfieldSpeakerPositionalData& InputPositions, const FAmbisonicsSoundfieldSettings& Settings, FAmbisonicsSoundfieldBuffer& OutputData);
+	SOUNDFIELDRENDERING_API FSoundFieldEncoder();
+	SOUNDFIELDRENDERING_API void EncodeAudioDirectlyFromOutputPositions(const Audio::FAlignedFloatBuffer& InputData, const FSoundfieldSpeakerPositionalData& InputPositions, const FAmbisonicsSoundfieldSettings& Settings, FAmbisonicsSoundfieldBuffer& OutputData);
 
 private:
 	static void EncodeLoop(const int32 NumFrames, const int32 NumInputChannels, const float* RESTRICT InputAudioPtr, const int32 NumAmbiChannels, float* RESTRICT SpeakerGainsPtr, float* RESTRICT OutputAmbiBuffer);
diff --git a/Engine/Source/Runtime/SoundFieldRendering/Public/SphericalHarmonicCalculator.h b/Engine/Source/Runtime/SoundFieldRendering/Public/SphericalHarmonicCalculator.h
--- a/Engine/Source/Runtime/SoundFieldRendering/Public/SphericalHarmonicCalculator.h
+++ b/Engine/Source/Runtime/SoundFieldRendering/Public/SphericalHarmonicCalculator.h
@@ -15,7 +15,7 @@
 	Elevation is 0 horizontal plane, + is above horizontal plane
 */
 
-class SOUNDFIELDRENDERING_API FSphericalHarmonicCalculator
+class FSphericalHarmonicCalculator
 {
 public:
 	enum AmbiChanNumber
@@ -28,10 +28,10 @@ public:
 		/* 5th-Order */	ACN_25, ACN_26, ACN_27, ACN_28, ACN_29, ACN_30, ACN_31, ACN_32, ACN_33, ACN_34, ACN_35
 	};
 
-	static void ComputeSoundfieldChannelGains(const int32 Order, const float Azimuth, const float Elevation, float* OutGains);
+	static SOUNDFIELDRENDERING_API void ComputeSoundfieldChannelGains(const int32 Order, const float Azimuth, const float Elevation, float* OutGains);
 
-	static void GenerateFirstOrderRotationMatrixGivenRadians(const float RotXRadians, const float RotYRadians, const float RotZRadians, FMatrix& OutMatrix);
-	static void GenerateFirstOrderRotationMatrixGivenDegrees(const float RotXDegrees, const float RotYDegrees, const float RotZDegrees, FMatrix& OutMatrix);
+	static SOUNDFIELDRENDERING_API void GenerateFirstOrderRotationMatrixGivenRadians(const float RotXRadians, const float RotYRadians, const float RotZRadians, FMatrix& OutMatrix);
+	static SOUNDFIELDRENDERING_API void GenerateFirstOrderRotationMatrixGivenDegrees(const float RotXDegrees, const float RotYDegrees, const float RotZDegrees, FMatrix& OutMatrix);
 
 	static void AdjustUESphericalCoordinatesForAmbisonics(FVector2D& InOutVector)
 	{
diff --git a/Engine/Source/Runtime/StaticMeshDescription/Public/StaticMeshDescription.h b/Engine/Source/Runtime/StaticMeshDescription/Public/StaticMeshDescription.h
--- a/Engine/Source/Runtime/StaticMeshDescription/Public/StaticMeshDescription.h
+++ b/Engine/Source/Runtime/StaticMeshDescription/Public/StaticMeshDescription.h
@@ -25,14 +25,14 @@ PRAGMA_DISABLE_DEPRECATION_WARNINGS
 /**
 * A wrapper for MeshDescription, customized for static meshes
 */
-UCLASS(BlueprintType)
-class STATICMESHDESCRIPTION_API UStaticMeshDescription : public UMeshDescriptionBase
+UCLASS(BlueprintType, MinimalAPI)
+class UStaticMeshDescription : public UMeshDescriptionBase
 {
 public:
 	GENERATED_BODY()
 
 	/** Register attributes required by static mesh description */
-	virtual void RegisterAttributes() override;
+	STATICMESHDESCRIPTION_API virtual void RegisterAttributes() override;
 
 	virtual FStaticMeshAttributes& GetRequiredAttributes() override
 	{ 
@@ -45,13 +45,13 @@ public:
 	}
 
 	UFUNCTION(BlueprintPure, Category="MeshDescription")
-	FVector2D GetVertexInstanceUV(FVertexInstanceID VertexInstanceID, int32 UVIndex = 0) const;
+	STATICMESHDESCRIPTION_API FVector2D GetVertexInstanceUV(FVertexInstanceID VertexInstanceID, int32 UVIndex = 0) const;
 
 	UFUNCTION(BlueprintCallable, Category="MeshDescription")
-	void SetVertexInstanceUV(FVertexInstanceID VertexInstanceID, FVector2D UV, int32 UVIndex = 0);
+	STATICMESHDESCRIPTION_API void SetVertexInstanceUV(FVertexInstanceID VertexInstanceID, FVector2D UV, int32 UVIndex = 0);
 
 	UFUNCTION(BlueprintCallable, Category="MeshDescription")
-	void CreateCube(FVector Center, FVector HalfExtents, FPolygonGroupID PolygonGroup,
+	STATICMESHDESCRIPTION_API void CreateCube(FVector Center, FVector HalfExtents, FPolygonGroupID PolygonGroup,
 					FPolygonID& PolygonID_PlusX,
 					FPolygonID& PolygonID_MinusX,
 					FPolygonID& PolygonID_PlusY,
@@ -60,7 +60,7 @@ public:
 					FPolygonID& PolygonID_MinusZ);
 
 	UFUNCTION(BlueprintCallable, Category="MeshDescription")
-	void SetPolygonGroupMaterialSlotName(FPolygonGroupID PolygonGroupID, const FName& SlotName);
+	STATICMESHDESCRIPTION_API void SetPolygonGroupMaterialSlotName(FPolygonGroupID PolygonGroupID, const FName& SlotName);
 
 public:
 
diff --git a/Engine/Source/Runtime/StaticMeshDescription/Public/StaticMeshOperations.h b/Engine/Source/Runtime/StaticMeshDescription/Public/StaticMeshOperations.h
--- a/Engine/Source/Runtime/StaticMeshDescription/Public/StaticMeshOperations.h
+++ b/Engine/Source/Runtime/StaticMeshDescription/Public/StaticMeshOperations.h
@@ -43,7 +43,7 @@ enum class EComputeNTBsFlags : uint32
 };
 ENUM_CLASS_FLAGS(EComputeNTBsFlags);
 
-class STATICMESHDESCRIPTION_API FStaticMeshOperations
+class FStaticMeshOperations
 {
 public:
 	struct FAppendSettings
@@ -68,45 +68,45 @@ public:
 
 	/** Set the polygon tangent, normal, binormal and polygonCenter for all polygons in the mesh description. */
 	UE_DEPRECATED(4.26, "Please use ComputeTriangleTangentsAndNormals() instead.")
-	static void ComputePolygonTangentsAndNormals(FMeshDescription& MeshDescription, float ComparisonThreshold = 0.0f);
+	static STATICMESHDESCRIPTION_API void ComputePolygonTangentsAndNormals(FMeshDescription& MeshDescription, float ComparisonThreshold = 0.0f);
 
 	/** Set the triangle tangent, normal, binormal and triangleCenter for all triangles in the mesh description. */
-	static void ComputeTriangleTangentsAndNormals(FMeshDescription& MeshDescription, float ComparisonThreshold = 0.0f, const TCHAR* DebugName = nullptr);
+	static STATICMESHDESCRIPTION_API void ComputeTriangleTangentsAndNormals(FMeshDescription& MeshDescription, float ComparisonThreshold = 0.0f, const TCHAR* DebugName = nullptr);
 
 	/** 
 	 * Recompute any invalid normal, tangent or Bi-Normal for every vertex in the mesh description with the given options.
 	 * If the EComputeNTBsFlags Normals or Tangents are set, the corresponding data will be force-recomputed.
 	 */
-	static void ComputeTangentsAndNormals(FMeshDescription& MeshDescription, EComputeNTBsFlags ComputeNTBsOptions);
+	static STATICMESHDESCRIPTION_API void ComputeTangentsAndNormals(FMeshDescription& MeshDescription, EComputeNTBsFlags ComputeNTBsOptions);
 
 	/*
 	 * Make sure all normals and tangents are valid. If not, recompute them.
 	 */
-	static void RecomputeNormalsAndTangentsIfNeeded(FMeshDescription& MeshDescription, EComputeNTBsFlags ComputeNTBsOptions);
+	static STATICMESHDESCRIPTION_API void RecomputeNormalsAndTangentsIfNeeded(FMeshDescription& MeshDescription, EComputeNTBsFlags ComputeNTBsOptions);
 
 	/** Compute tangent and Bi-Normal using mikkt space for every vertex in the mesh description. */
-	static void ComputeMikktTangents(FMeshDescription& MeshDescription, bool bIgnoreDegenerateTriangles);
+	static STATICMESHDESCRIPTION_API void ComputeMikktTangents(FMeshDescription& MeshDescription, bool bIgnoreDegenerateTriangles);
 
 	/** Determine the edge hardnesses from existing normals */
-	static void DetermineEdgeHardnessesFromVertexInstanceNormals(FMeshDescription& MeshDescription, float Tolerance = UE_KINDA_SMALL_NUMBER);
+	static STATICMESHDESCRIPTION_API void DetermineEdgeHardnessesFromVertexInstanceNormals(FMeshDescription& MeshDescription, float Tolerance = UE_KINDA_SMALL_NUMBER);
 
 	/** Convert this mesh description into the old FRawMesh format. */
-	static void ConvertToRawMesh(const FMeshDescription& SourceMeshDescription, FRawMesh& DestinationRawMesh, const TMap<FName, int32>& MaterialMap);
+	static STATICMESHDESCRIPTION_API void ConvertToRawMesh(const FMeshDescription& SourceMeshDescription, FRawMesh& DestinationRawMesh, const TMap<FName, int32>& MaterialMap);
 
 	/** Convert old FRawMesh format to MeshDescription. */
-	static void ConvertFromRawMesh(const FRawMesh& SourceRawMesh, FMeshDescription& DestinationMeshDescription, const TMap<int32, FName>& MaterialMap, bool bSkipNormalsAndTangents = false, const TCHAR* DebugName = nullptr);
+	static STATICMESHDESCRIPTION_API void ConvertFromRawMesh(const FRawMesh& SourceRawMesh, FMeshDescription& DestinationMeshDescription, const TMap<int32, FName>& MaterialMap, bool bSkipNormalsAndTangents = false, const TCHAR* DebugName = nullptr);
 
-	static void AppendMeshDescription(const FMeshDescription& SourceMesh, FMeshDescription& TargetMesh, FAppendSettings& AppendSettings);
-	static void AppendMeshDescriptions(const TArray<const FMeshDescription*>& SourceMeshes, FMeshDescription& TargetMesh, FAppendSettings& AppendSettings);
+	static STATICMESHDESCRIPTION_API void AppendMeshDescription(const FMeshDescription& SourceMesh, FMeshDescription& TargetMesh, FAppendSettings& AppendSettings);
+	static STATICMESHDESCRIPTION_API void AppendMeshDescriptions(const TArray<const FMeshDescription*>& SourceMeshes, FMeshDescription& TargetMesh, FAppendSettings& AppendSettings);
 
-	static void AreNormalsAndTangentsValid(const FMeshDescription& MeshDescription, bool& bHasInvalidNormals, bool& bHasInvalidTangents);
+	static STATICMESHDESCRIPTION_API void AreNormalsAndTangentsValid(const FMeshDescription& MeshDescription, bool& bHasInvalidNormals, bool& bHasInvalidTangents);
 
 
 	/** Find all overlapping vertex using the threshold in the mesh description. */
-	static void FindOverlappingCorners(FOverlappingCorners& OverlappingCorners, const FMeshDescription& MeshDescription, float ComparisonThreshold);
+	static STATICMESHDESCRIPTION_API void FindOverlappingCorners(FOverlappingCorners& OverlappingCorners, const FMeshDescription& MeshDescription, float ComparisonThreshold);
 
 	/** Find all charts in the mesh description. */
-	static int32 GetUVChartCount(FMeshDescription& MeshDescription, int32 SrcLightmapIndex, ELightmapUVVersion LightmapUVVersion, const FOverlappingCorners& OverlappingCorners);
+	static STATICMESHDESCRIPTION_API int32 GetUVChartCount(FMeshDescription& MeshDescription, int32 SrcLightmapIndex, ELightmapUVVersion LightmapUVVersion, const FOverlappingCorners& OverlappingCorners);
 
 	/**
 	 * Find and pack UV charts for lightmap.
@@ -122,7 +122,7 @@ public:
 	 * @param OverlappingCorners     Overlapping corners of the given mesh
 	 * @return                       UV layout correctly generated
 	 */
-	static bool CreateLightMapUVLayout(FMeshDescription& MeshDescription,
+	static STATICMESHDESCRIPTION_API bool CreateLightMapUVLayout(FMeshDescription& MeshDescription,
 		int32 SrcLightmapIndex,
 		int32 DstLightmapIndex,
 		int32 MinLightmapResolution,
@@ -130,25 +130,25 @@ public:
 		const FOverlappingCorners& OverlappingCorners);
 
 	/** Create some UVs from the specified mesh description data. */
-	static bool GenerateUniqueUVsForStaticMesh(const FMeshDescription& MeshDescription, int32 TextureResolution, bool bMergeIdenticalMaterials, TArray<FVector2D>& OutTexCoords);
+	static STATICMESHDESCRIPTION_API bool GenerateUniqueUVsForStaticMesh(const FMeshDescription& MeshDescription, int32 TextureResolution, bool bMergeIdenticalMaterials, TArray<FVector2D>& OutTexCoords);
 
 	/** Add a UV channel to the MeshDescription. */
-	static bool AddUVChannel(FMeshDescription& MeshDescription);
+	static STATICMESHDESCRIPTION_API bool AddUVChannel(FMeshDescription& MeshDescription);
 
 	/** Insert a UV channel at the given index to the MeshDescription. */
-	static bool InsertUVChannel(FMeshDescription& MeshDescription, int32 UVChannelIndex);
+	static STATICMESHDESCRIPTION_API bool InsertUVChannel(FMeshDescription& MeshDescription, int32 UVChannelIndex);
 
 	/** Remove the UV channel at the given index from the MeshDescription. */
-	static bool RemoveUVChannel(FMeshDescription& MeshDescription, int32 UVChannelIndex);
+	static STATICMESHDESCRIPTION_API bool RemoveUVChannel(FMeshDescription& MeshDescription, int32 UVChannelIndex);
 
 	/** Generate planar UV mapping for the MeshDescription */
-	static void GeneratePlanarUV(const FMeshDescription& MeshDescription, const FUVMapParameters& Params, TMap<FVertexInstanceID, FVector2D>& OutTexCoords);
+	static STATICMESHDESCRIPTION_API void GeneratePlanarUV(const FMeshDescription& MeshDescription, const FUVMapParameters& Params, TMap<FVertexInstanceID, FVector2D>& OutTexCoords);
 
 	/** Generate cylindrical UV mapping for the MeshDescription */
-	static void GenerateCylindricalUV(FMeshDescription& MeshDescription, const FUVMapParameters& Params, TMap<FVertexInstanceID, FVector2D>& OutTexCoords);
+	static STATICMESHDESCRIPTION_API void GenerateCylindricalUV(FMeshDescription& MeshDescription, const FUVMapParameters& Params, TMap<FVertexInstanceID, FVector2D>& OutTexCoords);
 
 	/** Generate box UV mapping for the MeshDescription */
-	static void GenerateBoxUV(const FMeshDescription& MeshDescription, const FUVMapParameters& Params, TMap<FVertexInstanceID, FVector2D>& OutTexCoords);
+	static STATICMESHDESCRIPTION_API void GenerateBoxUV(const FMeshDescription& MeshDescription, const FUVMapParameters& Params, TMap<FVertexInstanceID, FVector2D>& OutTexCoords);
 
 	//static void RemapPolygonGroups(FMeshDescription& MeshDescription, TMap<FPolygonGroupID, FPolygonGroupID>& Remap);
 
@@ -162,34 +162,34 @@ public:
 	 *                           If param is true  : PolygonGroupTargetID.GetValue() do not necessary equal SectionIndex in case there is less sections then SectionIndex
 	 *                           If param is false : PolygonGroupTargetID.GetValue() equal SectionIndex, we will add all necessary missing PolygonGroupID (this can generate empty PolygonGroupID)
 	 */
-	static void SwapPolygonPolygonGroup(FMeshDescription& MeshDescription, int32 SectionIndex, int32 TriangleIndexStart, int32 TriangleIndexEnd, bool bRemoveEmptyPolygonGroup);
+	static STATICMESHDESCRIPTION_API void SwapPolygonPolygonGroup(FMeshDescription& MeshDescription, int32 SectionIndex, int32 TriangleIndexStart, int32 TriangleIndexEnd, bool bRemoveEmptyPolygonGroup);
 
-	static void ConvertHardEdgesToSmoothGroup(const FMeshDescription& SourceMeshDescription, TArray<uint32>& FaceSmoothingMasks);
+	static STATICMESHDESCRIPTION_API void ConvertHardEdgesToSmoothGroup(const FMeshDescription& SourceMeshDescription, TArray<uint32>& FaceSmoothingMasks);
 
-	static void ConvertSmoothGroupToHardEdges(const TArray<uint32>& FaceSmoothingMasks, FMeshDescription& DestinationMeshDescription);
+	static STATICMESHDESCRIPTION_API void ConvertSmoothGroupToHardEdges(const TArray<uint32>& FaceSmoothingMasks, FMeshDescription& DestinationMeshDescription);
 
-	static bool HasVertexColor(const FMeshDescription& MeshDescription);
+	static STATICMESHDESCRIPTION_API bool HasVertexColor(const FMeshDescription& MeshDescription);
 
-	static void BuildWeldedVertexIDRemap(const FMeshDescription& MeshDescription, const float WeldingThreshold, TMap<FVertexID, FVertexID>& OutVertexIDRemap);
+	static STATICMESHDESCRIPTION_API void BuildWeldedVertexIDRemap(const FMeshDescription& MeshDescription, const float WeldingThreshold, TMap<FVertexID, FVertexID>& OutVertexIDRemap);
 
 	/** Computes the SHA hash of all the attributes values in the MeshDescription.
 	 * @param bSkipTransientAttributes     If param is true, do not include transient attributes in the hash computation.
 	 */
-	static FSHAHash ComputeSHAHash(const FMeshDescription& MeshDescription, bool bSkipTransientAttributes = false);
+	static STATICMESHDESCRIPTION_API FSHAHash ComputeSHAHash(const FMeshDescription& MeshDescription, bool bSkipTransientAttributes = false);
 
 	/** Flip the facing for a set of input polygons. */
-	static void FlipPolygons(FMeshDescription& MeshDescription);
+	static STATICMESHDESCRIPTION_API void FlipPolygons(FMeshDescription& MeshDescription);
 
 	/** Transforms the MeshDescription data using the provided transform. */
-	static void ApplyTransform(FMeshDescription& MeshDescription, const FTransform& Transform);
-	static void ApplyTransform(FMeshDescription& MeshDescription, const FMatrix& Transform);	
+	static STATICMESHDESCRIPTION_API void ApplyTransform(FMeshDescription& MeshDescription, const FTransform& Transform);
+	static STATICMESHDESCRIPTION_API void ApplyTransform(FMeshDescription& MeshDescription, const FMatrix& Transform);	
 
 	/**
 	 * Return the number of unique vertices, unique vertices are the result of welding all similar vertex instances (position, UV, tangent space, color,...)
 	 * 
 	 */
-	static int32 GetUniqueVertexCount(const FMeshDescription& MeshDescription);
-	static int32 GetUniqueVertexCount(const FMeshDescription& MeshDescription, const FOverlappingCorners& OverlappingCorners);
+	static STATICMESHDESCRIPTION_API int32 GetUniqueVertexCount(const FMeshDescription& MeshDescription);
+	static STATICMESHDESCRIPTION_API int32 GetUniqueVertexCount(const FMeshDescription& MeshDescription, const FOverlappingCorners& OverlappingCorners);
 
 	/**
 	 * Reorder the destination mesh description polygon groups in the same order has the source mesh description polygon groups.
@@ -197,7 +197,7 @@ public:
 	 * 
 	 * Note: The material count MUST be the same for both the source and the destination. It will not reorder anything if the count is different.
 	 */
-	static void ReorderMeshDescriptionPolygonGroups(const FMeshDescription& SourceMeshDescription
+	static STATICMESHDESCRIPTION_API void ReorderMeshDescriptionPolygonGroups(const FMeshDescription& SourceMeshDescription
 		, FMeshDescription& DestinationMeshDescription
 		, TOptional<const FString> UnmatchMaterialNameWarning
 		, TOptional<const FString> DestinationPolygonGroupCountDifferFromSource_Msg);
diff --git a/Engine/Source/Runtime/StaticMeshDescription/Public/UVMapSettings.h b/Engine/Source/Runtime/StaticMeshDescription/Public/UVMapSettings.h
--- a/Engine/Source/Runtime/StaticMeshDescription/Public/UVMapSettings.h
+++ b/Engine/Source/Runtime/StaticMeshDescription/Public/UVMapSettings.h
@@ -10,7 +10,7 @@
 
 /** UV map generation settings that are exposed to the user for scripting and through the editor */
 USTRUCT(BlueprintType)
-struct STATICMESHDESCRIPTION_API FUVMapSettings
+struct FUVMapSettings
 {
 	GENERATED_BODY()
 
@@ -45,7 +45,7 @@ struct STATICMESHDESCRIPTION_API FUVMapSettings
 	}
 };
 
-struct STATICMESHDESCRIPTION_API FUVMapParameters
+struct FUVMapParameters
 {
 	/** Length, width, height of the UV mapping gizmo */
 	FVector Size;
@@ -80,4 +80,4 @@ struct STATICMESHDESCRIPTION_API FUVMapParameters
 		, Scale(InScale)
 	{
 	}
-};
\ No newline at end of file
+};
diff --git a/Engine/Source/Runtime/StreamingFile/Public/StreamingNetworkPlatformFile.h b/Engine/Source/Runtime/StreamingFile/Public/StreamingNetworkPlatformFile.h
--- a/Engine/Source/Runtime/StreamingFile/Public/StreamingNetworkPlatformFile.h
+++ b/Engine/Source/Runtime/StreamingFile/Public/StreamingNetworkPlatformFile.h
@@ -21,7 +21,7 @@ DECLARE_LOG_CATEGORY_EXTERN(LogStreamingPlatformFile, Log, All);
 /**
  * Visitor to gather local files with their timestamps
  */
-class STREAMINGFILE_API FStreamingLocalTimestampVisitor : public IPlatformFile::FDirectoryVisitor
+class FStreamingLocalTimestampVisitor : public IPlatformFile::FDirectoryVisitor
 {
 private:
 
@@ -126,13 +126,13 @@ public:
 /**
  * Wrapper to redirect the low level file system to a server
  */
-class STREAMINGFILE_API FStreamingNetworkPlatformFile
+class FStreamingNetworkPlatformFile
 	: public FNetworkPlatformFile
 {
 	friend class FAsyncFileSync;
 
 	// FNetworkPlatformFile interface
-	virtual bool InitializeInternal(IPlatformFile* Inner, const TCHAR* HostIP) override;
+	STREAMINGFILE_API virtual bool InitializeInternal(IPlatformFile* Inner, const TCHAR* HostIP) override;
 
 public:
 
@@ -144,7 +144,7 @@ public:
 	}
 
 	/** Virtual destructor */
-	virtual ~FStreamingNetworkPlatformFile();
+	STREAMINGFILE_API virtual ~FStreamingNetworkPlatformFile();
 
 public:
 
@@ -154,19 +154,19 @@ public:
 	}
 
 	/** Sends Open message to the server and creates a new file handle if successful. */
-	class FStreamingNetworkFileHandle* SendOpenMessage(const FString& Filename, bool bIsWriting, bool bAppend, bool bAllowRead);
+	STREAMINGFILE_API class FStreamingNetworkFileHandle* SendOpenMessage(const FString& Filename, bool bIsWriting, bool bAppend, bool bAllowRead);
 
 	/** Sends Read message to the server. */
-	bool SendReadMessage(uint64 HandleId, uint8* Destination, int64 BytesToRead);
+	STREAMINGFILE_API bool SendReadMessage(uint64 HandleId, uint8* Destination, int64 BytesToRead);
 
 	/** Sends Write message to the server. */
-	bool SendWriteMessage(uint64 HandleId, const uint8* Source, int64 BytesToWrite);	
+	STREAMINGFILE_API bool SendWriteMessage(uint64 HandleId, const uint8* Source, int64 BytesToWrite);	
 
 	/** Sends Seek message to the server. */
-	bool SendSeekMessage(uint64 HandleId, int64 NewPosition);	
+	STREAMINGFILE_API bool SendSeekMessage(uint64 HandleId, int64 NewPosition);	
 
 	/** Sends Close message to the server. */
-	bool SendCloseMessage(uint64 HandleId);
+	STREAMINGFILE_API bool SendCloseMessage(uint64 HandleId);
 
 public:
 
@@ -175,7 +175,7 @@ public:
 
 	// IPlatformFile interface
 
-	virtual bool ShouldBeUsed(IPlatformFile* Inner, const TCHAR* CmdLine) const override;
+	STREAMINGFILE_API virtual bool ShouldBeUsed(IPlatformFile* Inner, const TCHAR* CmdLine) const override;
 
 	virtual IPlatformFile* GetLowerLevel() override
 	{
@@ -191,32 +191,32 @@ public:
 		return FStreamingNetworkPlatformFile::GetTypeName();
 	}
 
-	virtual bool DeleteFile(const TCHAR* Filename) override;
-	virtual bool IsReadOnly(const TCHAR* Filename) override;
-	virtual bool MoveFile(const TCHAR* To, const TCHAR* From) override;
-	virtual bool SetReadOnly(const TCHAR* Filename, bool bNewReadOnlyValue) override;
-	virtual FDateTime GetTimeStamp(const TCHAR* Filename) override;
-	virtual void SetTimeStamp(const TCHAR* Filename, FDateTime DateTime) override;
-	virtual FDateTime GetAccessTimeStamp(const TCHAR* Filename) override;
-	virtual IFileHandle* OpenRead(const TCHAR* Filename, bool bAllowWrite = false) override;
-	virtual IFileHandle* OpenWrite(const TCHAR* Filename, bool bAppend, bool bAllowRead) override;
-	virtual bool DirectoryExists(const TCHAR* Directory) override;
-	virtual bool CreateDirectoryTree(const TCHAR* Directory) override;
-	virtual bool CreateDirectory(const TCHAR* Directory) override;
-	virtual bool DeleteDirectory(const TCHAR* Directory) override;
-	virtual bool IterateDirectory(const TCHAR* Directory, IPlatformFile::FDirectoryVisitor& Visitor) override;
-	virtual bool IterateDirectoryRecursively(const TCHAR* Directory, IPlatformFile::FDirectoryVisitor& Visitor) override;
-	virtual bool DeleteDirectoryRecursively(const TCHAR* Directory) override;
-	virtual bool CopyFile(const TCHAR* To, const TCHAR* From, EPlatformFileRead ReadFlags = EPlatformFileRead::None, EPlatformFileWrite WriteFlags = EPlatformFileWrite::None) override;
-	virtual FString ConvertToAbsolutePathForExternalAppForRead( const TCHAR* Filename ) override;
-	virtual FString ConvertToAbsolutePathForExternalAppForWrite( const TCHAR* Filename ) override;
+	STREAMINGFILE_API virtual bool DeleteFile(const TCHAR* Filename) override;
+	STREAMINGFILE_API virtual bool IsReadOnly(const TCHAR* Filename) override;
+	STREAMINGFILE_API virtual bool MoveFile(const TCHAR* To, const TCHAR* From) override;
+	STREAMINGFILE_API virtual bool SetReadOnly(const TCHAR* Filename, bool bNewReadOnlyValue) override;
+	STREAMINGFILE_API virtual FDateTime GetTimeStamp(const TCHAR* Filename) override;
+	STREAMINGFILE_API virtual void SetTimeStamp(const TCHAR* Filename, FDateTime DateTime) override;
+	STREAMINGFILE_API virtual FDateTime GetAccessTimeStamp(const TCHAR* Filename) override;
+	STREAMINGFILE_API virtual IFileHandle* OpenRead(const TCHAR* Filename, bool bAllowWrite = false) override;
+	STREAMINGFILE_API virtual IFileHandle* OpenWrite(const TCHAR* Filename, bool bAppend, bool bAllowRead) override;
+	STREAMINGFILE_API virtual bool DirectoryExists(const TCHAR* Directory) override;
+	STREAMINGFILE_API virtual bool CreateDirectoryTree(const TCHAR* Directory) override;
+	STREAMINGFILE_API virtual bool CreateDirectory(const TCHAR* Directory) override;
+	STREAMINGFILE_API virtual bool DeleteDirectory(const TCHAR* Directory) override;
+	STREAMINGFILE_API virtual bool IterateDirectory(const TCHAR* Directory, IPlatformFile::FDirectoryVisitor& Visitor) override;
+	STREAMINGFILE_API virtual bool IterateDirectoryRecursively(const TCHAR* Directory, IPlatformFile::FDirectoryVisitor& Visitor) override;
+	STREAMINGFILE_API virtual bool DeleteDirectoryRecursively(const TCHAR* Directory) override;
+	STREAMINGFILE_API virtual bool CopyFile(const TCHAR* To, const TCHAR* From, EPlatformFileRead ReadFlags = EPlatformFileRead::None, EPlatformFileWrite WriteFlags = EPlatformFileWrite::None) override;
+	STREAMINGFILE_API virtual FString ConvertToAbsolutePathForExternalAppForRead( const TCHAR* Filename ) override;
+	STREAMINGFILE_API virtual FString ConvertToAbsolutePathForExternalAppForWrite( const TCHAR* Filename ) override;
 
 private:
 
 	// FNetworkPlatformFile interface
 
-	virtual void PerformHeartbeat() override;
-	virtual void GetFileInfo(const TCHAR* Filename, FFileInfo& Info) override;
+	STREAMINGFILE_API virtual void PerformHeartbeat() override;
+	STREAMINGFILE_API virtual void GetFileInfo(const TCHAR* Filename, FFileInfo& Info) override;
 
 private:
 
diff --git a/Engine/Source/Runtime/TextureUtilitiesCommon/Public/TextureImportSettings.h b/Engine/Source/Runtime/TextureUtilitiesCommon/Public/TextureImportSettings.h
--- a/Engine/Source/Runtime/TextureUtilitiesCommon/Public/TextureImportSettings.h
+++ b/Engine/Source/Runtime/TextureUtilitiesCommon/Public/TextureImportSettings.h
@@ -46,8 +46,8 @@ enum class ETextureImportPNGInfill : uint8
 	Always
 };
 
-UCLASS(config=Editor, defaultconfig, meta=(DisplayName="Texture Import"))
-class TEXTUREUTILITIESCOMMON_API UTextureImportSettings : public UDeveloperSettings
+UCLASS(config=Editor, defaultconfig, meta=(DisplayName="Texture Import"), MinimalAPI)
+class UTextureImportSettings : public UDeveloperSettings
 {
 	GENERATED_UCLASS_BODY()
 	
@@ -80,13 +80,13 @@ public:
 
 	//~ Begin UObject Interface
 
-	virtual void PostInitProperties() override;
+	TEXTUREUTILITIESCOMMON_API virtual void PostInitProperties() override;
 	
 	// Get the PNGInfill setting, with Default mapped to a concrete choice
-	ETextureImportPNGInfill GetPNGInfillMapDefault() const;
+	TEXTUREUTILITIESCOMMON_API ETextureImportPNGInfill GetPNGInfillMapDefault() const;
 
 #if WITH_EDITOR
-	virtual void PostEditChangeProperty(FPropertyChangedEvent& PropertyChangedEvent) override;
+	TEXTUREUTILITIESCOMMON_API virtual void PostEditChangeProperty(FPropertyChangedEvent& PropertyChangedEvent) override;
 #endif
 
 	//~ End UObject Interface
@@ -104,4 +104,4 @@ namespace UE::TextureUtilitiesCommon
 	TEXTUREUTILITIESCOMMON_API void ApplyDefaultsForNewlyImportedTextures(UTexture * Texture, bool bIsReimport);
 
 }
-#endif
\ No newline at end of file
+#endif
diff --git a/Engine/Source/Runtime/TimeManagement/Public/CommonFrameRates.h b/Engine/Source/Runtime/TimeManagement/Public/CommonFrameRates.h
--- a/Engine/Source/Runtime/TimeManagement/Public/CommonFrameRates.h
+++ b/Engine/Source/Runtime/TimeManagement/Public/CommonFrameRates.h
@@ -35,7 +35,7 @@ struct FCommonFrameRateInfo
 	FText Description;
 };
 
-struct TIMEMANAGEMENT_API FCommonFrameRates
+struct FCommonFrameRates
 {
 	typedef __underlying_type(ECommonFrameRate) NumericType;
 
@@ -55,14 +55,14 @@ struct TIMEMANAGEMENT_API FCommonFrameRates
 	FORCEINLINE static FFrameRate NTSC_30() { return AllFrameRates[(NumericType)ECommonFrameRate::NTSC_30].FrameRate; }
 	FORCEINLINE static FFrameRate NTSC_60() { return AllFrameRates[(NumericType)ECommonFrameRate::NTSC_60].FrameRate; }
 
-	static TArrayView<const FCommonFrameRateInfo> GetAll();
+	static TIMEMANAGEMENT_API TArrayView<const FCommonFrameRateInfo> GetAll();
 
 	static bool Contains(FFrameRate FrameRateToCheck)
 	{
 		return Find(FrameRateToCheck) != nullptr;
 	}
 
-	static const FCommonFrameRateInfo* Find(FFrameRate InFrameRate);
+	static TIMEMANAGEMENT_API const FCommonFrameRateInfo* Find(FFrameRate InFrameRate);
 
 	/** Find a common frame rate that matches the given frame rate as a decimal number of frames per second.
 	 *
@@ -70,8 +70,8 @@ struct TIMEMANAGEMENT_API FCommonFrameRates
 	 *  @param Tolerance: Numerical tolerance to use when searching for a frame rate match.
 	 *  @return: a pointer to the matching common frame rate if a match was found, or nullptr otherwise.
 	 */
-	static const FCommonFrameRateInfo* Find(const double InFrameRateAsDecimal, const double Tolerance = UE_DOUBLE_KINDA_SMALL_NUMBER);
+	static TIMEMANAGEMENT_API const FCommonFrameRateInfo* Find(const double InFrameRateAsDecimal, const double Tolerance = UE_DOUBLE_KINDA_SMALL_NUMBER);
 
 private:
-	static const FCommonFrameRateInfo AllFrameRates[(int32)ECommonFrameRate::Private_Num];
-};
\ No newline at end of file
+	static TIMEMANAGEMENT_API const FCommonFrameRateInfo AllFrameRates[(int32)ECommonFrameRate::Private_Num];
+};
diff --git a/Engine/Source/Runtime/TimeManagement/Public/FixedFrameRateCustomTimeStep.h b/Engine/Source/Runtime/TimeManagement/Public/FixedFrameRateCustomTimeStep.h
--- a/Engine/Source/Runtime/TimeManagement/Public/FixedFrameRateCustomTimeStep.h
+++ b/Engine/Source/Runtime/TimeManagement/Public/FixedFrameRateCustomTimeStep.h
@@ -18,18 +18,18 @@ class UObject;
 /**
  * Class to control the Engine TimeStep via a FixedFrameRate
  */
-UCLASS(Abstract)
-class TIMEMANAGEMENT_API UFixedFrameRateCustomTimeStep : public UEngineCustomTimeStep
+UCLASS(Abstract, MinimalAPI)
+class UFixedFrameRateCustomTimeStep : public UEngineCustomTimeStep
 {
 	GENERATED_UCLASS_BODY()
 
 public:
 	/** Get The fixed FrameRate */
-	virtual FFrameRate GetFixedFrameRate() const PURE_VIRTUAL(UFixedFrameRateCustomTimeStep::GetFixedFrameRate, return GetFixedFrameRate_PureVirtual(););
+	TIMEMANAGEMENT_API virtual FFrameRate GetFixedFrameRate() const PURE_VIRTUAL(UFixedFrameRateCustomTimeStep::GetFixedFrameRate, return GetFixedFrameRate_PureVirtual(););
 
 protected:
 	/** Default behavior of the engine. Used FixedFrameRate */
-	void WaitForFixedFrameRate() const;
+	TIMEMANAGEMENT_API void WaitForFixedFrameRate() const;
 
 private:
 	FFrameRate GetFixedFrameRate_PureVirtual() const;
diff --git a/Engine/Source/Runtime/TimeManagement/Public/FrameNumberTimeEvaluator.h b/Engine/Source/Runtime/TimeManagement/Public/FrameNumberTimeEvaluator.h
--- a/Engine/Source/Runtime/TimeManagement/Public/FrameNumberTimeEvaluator.h
+++ b/Engine/Source/Runtime/TimeManagement/Public/FrameNumberTimeEvaluator.h
@@ -21,15 +21,15 @@ DEFINE_EXPRESSION_OPERATOR_NODE(TIMEMANAGEMENT_API, FDropcodeDelimiter, 0x80AF2C
 DEFINE_EXPRESSION_OPERATOR_NODE(TIMEMANAGEMENT_API, FBracketStart, 0xA7358BD1, 0xB4EF466D, 0xA336CD84, 0x14C93D2E);
 DEFINE_EXPRESSION_OPERATOR_NODE(TIMEMANAGEMENT_API, FBracketEnd, 0xE49500E9, 0x03E64440, 0x87802630, 0xF1C1DDDF);
 
-struct TIMEMANAGEMENT_API FFrameNumberTimeEvaluator
+struct FFrameNumberTimeEvaluator
 {
 	/** Constructor that sets up the parser's lexer and compiler */
-	FFrameNumberTimeEvaluator();
+	TIMEMANAGEMENT_API FFrameNumberTimeEvaluator();
 
 	/** Evaluate the given expression, resulting in either a double value, or an error */
-	TValueOrError<FFrameTime, FExpressionError> EvaluateTimecode(const TCHAR* InExpression, const FFrameRate& InDisplayFrameRate, const FFrameRate& InTickResolution, bool& OutDirectlyParsed) const;
-	TValueOrError<FFrameTime, FExpressionError> EvaluateFrame(const TCHAR* InExpression, const FFrameRate& InDisplayFrameRate, const FFrameRate& InTickResolution, bool& OutDirectlyParsed) const;
-	TValueOrError<FFrameTime, FExpressionError> EvaluateTime(const TCHAR* InExpression, FFrameRate InFrameRate, bool& OutDirectlyParsed) const;
+	TIMEMANAGEMENT_API TValueOrError<FFrameTime, FExpressionError> EvaluateTimecode(const TCHAR* InExpression, const FFrameRate& InDisplayFrameRate, const FFrameRate& InTickResolution, bool& OutDirectlyParsed) const;
+	TIMEMANAGEMENT_API TValueOrError<FFrameTime, FExpressionError> EvaluateFrame(const TCHAR* InExpression, const FFrameRate& InDisplayFrameRate, const FFrameRate& InTickResolution, bool& OutDirectlyParsed) const;
+	TIMEMANAGEMENT_API TValueOrError<FFrameTime, FExpressionError> EvaluateTime(const TCHAR* InExpression, FFrameRate InFrameRate, bool& OutDirectlyParsed) const;
 
 	FTokenDefinitions TimecodeTokenDefinitions;
 	FTokenDefinitions FrameTokenDefinitions;
@@ -40,4 +40,4 @@ struct TIMEMANAGEMENT_API FFrameNumberTimeEvaluator
 	FExpressionGrammar TimeGrammar;
 
 	FOperatorJumpTable TimeJumpTable;
-};
\ No newline at end of file
+};
diff --git a/Engine/Source/Runtime/TimeManagement/Public/GenlockedCustomTimeStep.h b/Engine/Source/Runtime/TimeManagement/Public/GenlockedCustomTimeStep.h
--- a/Engine/Source/Runtime/TimeManagement/Public/GenlockedCustomTimeStep.h
+++ b/Engine/Source/Runtime/TimeManagement/Public/GenlockedCustomTimeStep.h
@@ -17,34 +17,34 @@ class UObject;
 /**
  * Class to control the Engine Timestep from a Genlock signal.
  */
-UCLASS(Abstract)
-class TIMEMANAGEMENT_API UGenlockedCustomTimeStep : public UFixedFrameRateCustomTimeStep
+UCLASS(Abstract, MinimalAPI)
+class UGenlockedCustomTimeStep : public UFixedFrameRateCustomTimeStep
 {
 	GENERATED_BODY()
 
 protected:
 
 	/** Blocks until it gets a sync signal. Returns false if unsuccessful */
-	virtual bool WaitForSync() PURE_VIRTUAL(UGenlockedCustomTimeStep::WaitForSync, return false;);
+	TIMEMANAGEMENT_API virtual bool WaitForSync() PURE_VIRTUAL(UGenlockedCustomTimeStep::WaitForSync, return false;);
 
 	/** Update FApp CurrentTime, IdleTime and DeltaTime */
-	void UpdateAppTimes(const double& TimeBeforeSync, const double& TimeAfterSync) const;
+	TIMEMANAGEMENT_API void UpdateAppTimes(const double& TimeBeforeSync, const double& TimeAfterSync) const;
 
 public:
 	//~ UFixedFrameRateCustomTimeStep interface
-	virtual FFrameRate GetFixedFrameRate() const PURE_VIRTUAL(UGenlockedCustomTimeStep::GetFixedFrameRate, return FFrameRate(24,1););
+	TIMEMANAGEMENT_API virtual FFrameRate GetFixedFrameRate() const PURE_VIRTUAL(UGenlockedCustomTimeStep::GetFixedFrameRate, return FFrameRate(24,1););
 
 	/** Get the sync rate (not always the same as the fixed frame rate) */
 	virtual FFrameRate GetSyncRate() const { return GetFixedFrameRate(); };
 
 	/** Returns how many syncs occurred since the last tick */
-	virtual uint32 GetLastSyncCountDelta() const PURE_VIRTUAL(UGenlockedCustomTimeStep::GetLastSyncCountDelta, return 1;);
+	TIMEMANAGEMENT_API virtual uint32 GetLastSyncCountDelta() const PURE_VIRTUAL(UGenlockedCustomTimeStep::GetLastSyncCountDelta, return 1;);
 
 	/** Returns true if the Sync related functions will return valid data */
-	virtual bool IsLastSyncDataValid() const PURE_VIRTUAL(UGenlockedCustomTimeStep::IsLastSyncDataValid, return false;);
+	TIMEMANAGEMENT_API virtual bool IsLastSyncDataValid() const PURE_VIRTUAL(UGenlockedCustomTimeStep::IsLastSyncDataValid, return false;);
 
 	/** Returns how many sync counts are expected between ticks*/
-	virtual uint32 GetExpectedSyncCountDelta() const;
+	TIMEMANAGEMENT_API virtual uint32 GetExpectedSyncCountDelta() const;
 
 	/** Whether automatic format detection is supported. */ 
 	virtual bool SupportsFormatAutoDetection() const { return false; };
@@ -54,5 +54,5 @@ public:
 	bool bAutoDetectFormat = false;
 
 	/** Experimental fix for interlace field flipping issue. */
-	static TAutoConsoleVariable<int32> CVarExperimentalFieldFlipFix;
+	static TIMEMANAGEMENT_API TAutoConsoleVariable<int32> CVarExperimentalFieldFlipFix;
 };
diff --git a/Engine/Source/Runtime/TimeManagement/Public/GenlockedFixedRateCustomTimeStep.h b/Engine/Source/Runtime/TimeManagement/Public/GenlockedFixedRateCustomTimeStep.h
--- a/Engine/Source/Runtime/TimeManagement/Public/GenlockedFixedRateCustomTimeStep.h
+++ b/Engine/Source/Runtime/TimeManagement/Public/GenlockedFixedRateCustomTimeStep.h
@@ -29,23 +29,23 @@ class UObject;
  *       - Corresponds to platform time minus any fractional FrameTime.
  * 
  */
-UCLASS(Blueprintable, editinlinenew, meta = (DisplayName = "Genlocked Fixed Rate"))
-class TIMEMANAGEMENT_API UGenlockedFixedRateCustomTimeStep : public UGenlockedCustomTimeStep
+UCLASS(Blueprintable, editinlinenew, meta = (DisplayName = "Genlocked Fixed Rate"), MinimalAPI)
+class UGenlockedFixedRateCustomTimeStep : public UGenlockedCustomTimeStep
 {
 	GENERATED_UCLASS_BODY()
 
 public:
 	//~ UFixedFrameRateCustomTimeStep interface
-	virtual bool Initialize(UEngine* InEngine) override;
-	virtual void Shutdown(UEngine* InEngine) override;
-	virtual bool UpdateTimeStep(UEngine* InEngine) override;
-	virtual ECustomTimeStepSynchronizationState GetSynchronizationState() const override;
-	virtual FFrameRate GetFixedFrameRate() const override;
+	TIMEMANAGEMENT_API virtual bool Initialize(UEngine* InEngine) override;
+	TIMEMANAGEMENT_API virtual void Shutdown(UEngine* InEngine) override;
+	TIMEMANAGEMENT_API virtual bool UpdateTimeStep(UEngine* InEngine) override;
+	TIMEMANAGEMENT_API virtual ECustomTimeStepSynchronizationState GetSynchronizationState() const override;
+	TIMEMANAGEMENT_API virtual FFrameRate GetFixedFrameRate() const override;
 
 	//~ UGenlockedCustomTimeStep interface
-	virtual uint32 GetLastSyncCountDelta() const override;
-	virtual bool IsLastSyncDataValid() const override;
-	virtual bool WaitForSync() override;
+	TIMEMANAGEMENT_API virtual uint32 GetLastSyncCountDelta() const override;
+	TIMEMANAGEMENT_API virtual bool IsLastSyncDataValid() const override;
+	TIMEMANAGEMENT_API virtual bool WaitForSync() override;
 
 public:
 
diff --git a/Engine/Source/Runtime/TimeManagement/Public/GenlockedTimecodeProvider.h b/Engine/Source/Runtime/TimeManagement/Public/GenlockedTimecodeProvider.h
--- a/Engine/Source/Runtime/TimeManagement/Public/GenlockedTimecodeProvider.h
+++ b/Engine/Source/Runtime/TimeManagement/Public/GenlockedTimecodeProvider.h
@@ -11,8 +11,8 @@ class UGenlockedCustomTimeStep;
 /**
  * This timecode provider base class will try to use the engine genlock sync to adjust its count.
  */
-UCLASS(Abstract)
-class TIMEMANAGEMENT_API UGenlockedTimecodeProvider : public UTimecodeProvider
+UCLASS(Abstract, MinimalAPI)
+class UGenlockedTimecodeProvider : public UTimecodeProvider
 {
 	GENERATED_BODY()
 
@@ -23,14 +23,14 @@ public:
 	bool bUseGenlockToCount = true;
 
 	//~ Begin UTimecodeProvider interface
-	virtual void FetchAndUpdate() override;
-	virtual FQualifiedFrameTime GetQualifiedFrameTime() const override;
+	TIMEMANAGEMENT_API virtual void FetchAndUpdate() override;
+	TIMEMANAGEMENT_API virtual FQualifiedFrameTime GetQualifiedFrameTime() const override;
 	//~ End UTimecodeProvider interface
 
 protected:
 
 	/** Corrects given timecode with Genlock provider */
-	virtual FQualifiedFrameTime CorrectFromGenlock(FQualifiedFrameTime& InFrameTime, const UGenlockedCustomTimeStep* Genlock);
+	TIMEMANAGEMENT_API virtual FQualifiedFrameTime CorrectFromGenlock(FQualifiedFrameTime& InFrameTime, const UGenlockedCustomTimeStep* Genlock);
 
 	/** Cache current frame time */
 	FQualifiedFrameTime LastFrameTime;
diff --git a/Engine/Source/Runtime/TimeManagement/Public/ITimeManagementModule.h b/Engine/Source/Runtime/TimeManagement/Public/ITimeManagementModule.h
--- a/Engine/Source/Runtime/TimeManagement/Public/ITimeManagementModule.h
+++ b/Engine/Source/Runtime/TimeManagement/Public/ITimeManagementModule.h
@@ -11,7 +11,7 @@
 
 class FTimedDataInputCollection;
 
-class TIMEMANAGEMENT_API ITimeManagementModule : public IModuleInterface
+class ITimeManagementModule : public IModuleInterface
 {
 public:
 	
diff --git a/Engine/Source/Runtime/TimeManagement/Public/ITimedDataInput.h b/Engine/Source/Runtime/TimeManagement/Public/ITimedDataInput.h
--- a/Engine/Source/Runtime/TimeManagement/Public/ITimedDataInput.h
+++ b/Engine/Source/Runtime/TimeManagement/Public/ITimedDataInput.h
@@ -62,7 +62,7 @@ struct FTimedDataChannelSampleTime
 
 
 USTRUCT(BlueprintType)
-struct TIMEMANAGEMENT_API FTimedDataInputEvaluationData
+struct FTimedDataInputEvaluationData
 {
 	GENERATED_BODY()
 
@@ -83,13 +83,13 @@ struct TIMEMANAGEMENT_API FTimedDataInputEvaluationData
 /**
  * Interface for data sources that can be synchronized with time
  */
-class TIMEMANAGEMENT_API ITimedDataInput
+class ITimedDataInput
 {
 public:
-	static FFrameRate UnknownFrameRate;
+	static TIMEMANAGEMENT_API FFrameRate UnknownFrameRate;
 	
-	static double ConvertSecondOffsetInFrameOffset(double Seconds, FFrameRate Rate);
-	static double ConvertFrameOffsetInSecondOffset(double Frames, FFrameRate Rate);
+	static TIMEMANAGEMENT_API double ConvertSecondOffsetInFrameOffset(double Seconds, FFrameRate Rate);
+	static TIMEMANAGEMENT_API double ConvertFrameOffsetInSecondOffset(double Frames, FFrameRate Rate);
 	
 public:
 	virtual ~ITimedDataInput() {}
@@ -137,10 +137,10 @@ public:
 	}
 
 	/** Convert second offset to frame offset using this input's framerate. */
-	double ConvertSecondOffsetInFrameOffset(double Seconds) const;
+	TIMEMANAGEMENT_API double ConvertSecondOffsetInFrameOffset(double Seconds) const;
 
 	/** Convert frame offset to second offset using this input's framerate. */
-	double ConvertFrameOffsetInSecondOffset(double Frames) const;
+	TIMEMANAGEMENT_API double ConvertFrameOffsetInSecondOffset(double Frames) const;
 
 #if WITH_EDITOR
 	/** Get the icon that represent the input. */
@@ -152,7 +152,7 @@ public:
 /**
  * Interface for data tracked produced by an input.
  */
-class TIMEMANAGEMENT_API ITimedDataInputChannel
+class ITimedDataInputChannel
 {
 public:
 	/** Get the channel's display name. */
diff --git a/Engine/Source/Runtime/TimeManagement/Public/TimeManagementBlueprintLibrary.h b/Engine/Source/Runtime/TimeManagement/Public/TimeManagementBlueprintLibrary.h
--- a/Engine/Source/Runtime/TimeManagement/Public/TimeManagementBlueprintLibrary.h
+++ b/Engine/Source/Runtime/TimeManagement/Public/TimeManagementBlueprintLibrary.h
@@ -23,8 +23,8 @@ struct FQualifiedFrameTime;
 /**
  * 
  */
-UCLASS(meta = (BlueprintThreadSafe, ScriptName = "TimeManagementLibrary"))
-class TIMEMANAGEMENT_API UTimeManagementBlueprintLibrary : public UBlueprintFunctionLibrary
+UCLASS(meta = (BlueprintThreadSafe, ScriptName = "TimeManagementLibrary"), MinimalAPI)
+class UTimeManagementBlueprintLibrary : public UBlueprintFunctionLibrary
 {
 	GENERATED_BODY()
 
@@ -32,74 +32,74 @@ public:
 	/** Converts an FrameRate to a float ie: 1/30 returns 0.0333333 */
 	UE_DEPRECATED(5.3, "Conv_FrameRateToSeconds has been deprecated, use Conv_FrameRateToInterval instead")
 	UFUNCTION(BlueprintPure, meta = (DisplayName = "FrameRate To Seconds", BlueprintAutocast, DeprecatedFunction, DeprecationMessage = "FrameRateToInterval replaces this function, which returns the expected result of seconds per frame, rather than (incorrectly) frames per second."), Category = "Utilities|Time Management")
-	static float Conv_FrameRateToSeconds(const FFrameRate& InFrameRate);
+	static TIMEMANAGEMENT_API float Conv_FrameRateToSeconds(const FFrameRate& InFrameRate);
 
 	/** Converts a FrameRate to an interval float representing the frame time in seconds ie: 1/30 returns 0.0333333 */
 	UFUNCTION(BlueprintPure, meta = (DisplayName = "FrameRate To Interval", BlueprintAutocast), Category = "Utilities|Time Management")
-	static float Conv_FrameRateToInterval(const FFrameRate InFrameRate);
+	static TIMEMANAGEMENT_API float Conv_FrameRateToInterval(const FFrameRate InFrameRate);
 
 	/** Converts an QualifiedFrameTime to seconds. */
 	UFUNCTION(BlueprintPure, meta = (DisplayName = "QualifiedFrameTime To Seconds", BlueprintAutocast), Category = "Utilities|Time Management")
-	static float Conv_QualifiedFrameTimeToSeconds(const FQualifiedFrameTime& InFrameTime);
+	static TIMEMANAGEMENT_API float Conv_QualifiedFrameTimeToSeconds(const FQualifiedFrameTime& InFrameTime);
 
 	/** Multiplies a value in seconds against a FrameRate to get a new FrameTime. */
 	UFUNCTION(BlueprintPure, meta = (DisplayName = "Seconds * FrameRate", CompactNodeTitle = "*"), Category = "Utilities|Time Management")
-	static FFrameTime Multiply_SecondsFrameRate(float TimeInSeconds, const FFrameRate& FrameRate);
+	static TIMEMANAGEMENT_API FFrameTime Multiply_SecondsFrameRate(float TimeInSeconds, const FFrameRate& FrameRate);
 
 	/** Converts an Timecode to a string (hh:mm:ss:ff). If bForceSignDisplay then the number sign will always be prepended instead of just when expressing a negative time. */
 	UFUNCTION(BlueprintPure, meta = (DisplayName = "Timecode To String", BlueprintAutocast), Category = "Utilities|Time Management")
-	static FString Conv_TimecodeToString(const FTimecode& InTimecode, bool bForceSignDisplay = false);
+	static TIMEMANAGEMENT_API FString Conv_TimecodeToString(const FTimecode& InTimecode, bool bForceSignDisplay = false);
 
 	/** Verifies that this is a valid framerate with a non-zero denominator. */
 	UFUNCTION(BlueprintPure, meta = (DisplayName = "Is Valid Frame Rate"), Category = "Utilities|Time Management")
-	static bool IsValid_Framerate(const FFrameRate& InFrameRate);
+	static TIMEMANAGEMENT_API bool IsValid_Framerate(const FFrameRate& InFrameRate);
 
 	/** Checks if this framerate is an even multiple of another framerate, ie: 60 is a multiple of 30, but 59.94 is not. */
 	UFUNCTION(BlueprintPure, meta = (DisplayName = "Is Multiple Of"), Category = "Utilities|Time Management")
-	static bool IsValid_MultipleOf(const FFrameRate& InFrameRate, const FFrameRate& OtherFramerate);
+	static TIMEMANAGEMENT_API bool IsValid_MultipleOf(const FFrameRate& InFrameRate, const FFrameRate& OtherFramerate);
 
 	/** Converts the specified time from one framerate to another framerate. This is useful for converting between tick resolution and display rate. */
 	UFUNCTION(BlueprintPure, meta = (DisplayName = "Transform Frame Time"), Category = "Utilities|Time Management")
-	static FFrameTime TransformTime(const FFrameTime& SourceTime, const FFrameRate& SourceRate, const FFrameRate& DestinationRate);
+	static TIMEMANAGEMENT_API FFrameTime TransformTime(const FFrameTime& SourceTime, const FFrameRate& SourceRate, const FFrameRate& DestinationRate);
 
 	/** Snaps the given SourceTime to the nearest frame in the specified Destination Framerate. Useful for determining the nearest frame for another resolution. Returns the frame time in the destination frame rate. */
 	UFUNCTION(BlueprintPure, meta = (DisplayName = "Snap Frame Time"), Category = "Utilities|Time Management")
-	static FFrameTime SnapFrameTimeToRate(const FFrameTime& SourceTime, const FFrameRate& SourceRate, const FFrameRate& SnapToRate);
+	static TIMEMANAGEMENT_API FFrameTime SnapFrameTimeToRate(const FFrameTime& SourceTime, const FFrameRate& SourceRate, const FFrameRate& SnapToRate);
 
 	/** Addition (FrameNumber A + FrameNumber B) */
 	UFUNCTION(BlueprintPure, meta = (DisplayName = "FrameNumber + FrameNumber", CompactNodeTitle = "+", Keywords = "+ add plus", CommutativeAssociativeBinaryOperator = "true", ScriptMethod, ScriptMethodSelfReturn, ScriptOperator = "+;+="), Category = "Utilities|Time Management")
-	static FFrameNumber Add_FrameNumberFrameNumber(FFrameNumber A, FFrameNumber B);
+	static TIMEMANAGEMENT_API FFrameNumber Add_FrameNumberFrameNumber(FFrameNumber A, FFrameNumber B);
 
 	/** Subtraction (FrameNumber A - FrameNumber B) */
 	UFUNCTION(BlueprintPure, meta = (DisplayName = "FrameNumber - FrameNumber", CompactNodeTitle = "-", Keywords = "- subtract minus", ScriptMethod, ScriptMethodSelfReturn, ScriptOperator = "-;-="), Category = "Utilities|Time Management")
-	static FFrameNumber Subtract_FrameNumberFrameNumber(FFrameNumber A, FFrameNumber B);
+	static TIMEMANAGEMENT_API FFrameNumber Subtract_FrameNumberFrameNumber(FFrameNumber A, FFrameNumber B);
 
 	/** Addition (FrameNumber A + int B) */
 	UFUNCTION(BlueprintPure, meta = (DisplayName = "FrameNumber + Int", CompactNodeTitle = "+", Keywords = "+ add plus", ScriptMethod, ScriptMethodSelfReturn, ScriptOperator = "+;+="), Category = "Utilities|Time Management")
-	static FFrameNumber Add_FrameNumberInteger(FFrameNumber A, int32 B);
+	static TIMEMANAGEMENT_API FFrameNumber Add_FrameNumberInteger(FFrameNumber A, int32 B);
 
 	/** Subtraction (FrameNumber A - int B) */
 	UFUNCTION(BlueprintPure, meta = (DisplayName = "FrameNumber - Int", CompactNodeTitle = "-", Keywords = "- subtract minus", ScriptMethod, ScriptMethodSelfReturn, ScriptOperator = "-;-="), Category = "Utilities|Time Management")
-	static FFrameNumber Subtract_FrameNumberInteger(FFrameNumber A, int32 B);
+	static TIMEMANAGEMENT_API FFrameNumber Subtract_FrameNumberInteger(FFrameNumber A, int32 B);
 
 	/** Multiply (FrameNumber A * B) */
 	UFUNCTION(BlueprintPure, meta = (DisplayName = "FrameNumber * Int", CompactNodeTitle = "*", Keywords = "* multiply", ScriptMethod, ScriptMethodSelfReturn, ScriptOperator = "*;*="), Category = "Utilities|Time Management")
-	static FFrameNumber Multiply_FrameNumberInteger(FFrameNumber A, int32 B);
+	static TIMEMANAGEMENT_API FFrameNumber Multiply_FrameNumberInteger(FFrameNumber A, int32 B);
 
 	/** Divide (FrameNumber A / B) */
 	UFUNCTION(BlueprintPure, meta = (DisplayName = "FrameNumber / FrameNumber", CompactNodeTitle = "/", Keywords = "/ divide", ScriptMethod, ScriptMethodSelfReturn, ScriptOperator = "/;/="), Category = "Utilities|Time Management")
-	static FFrameNumber Divide_FrameNumberInteger(FFrameNumber A, int32 B);
+	static TIMEMANAGEMENT_API FFrameNumber Divide_FrameNumberInteger(FFrameNumber A, int32 B);
 	
 	/** Converts a FrameNumber to an int32 for use in functions that take int32 frame counts for convenience. */
 	UFUNCTION(BlueprintPure, meta = (DisplayName = "FrameNumber to Integer", ScriptName="FrameNumberToInteger", BlueprintAutocast), Category = "Utilities|Time Management")
-	static int32 Conv_FrameNumberToInteger(const FFrameNumber& InFrameNumber);
+	static TIMEMANAGEMENT_API int32 Conv_FrameNumberToInteger(const FFrameNumber& InFrameNumber);
 
 public:
 	/** Get the current timecode of the engine. */
 	UFUNCTION(BlueprintPure, Category = "Utilities|Time Management")
-	static FTimecode GetTimecode();
+	static TIMEMANAGEMENT_API FTimecode GetTimecode();
 
 	/** Gets the current timecode frame rate. */
 	UFUNCTION(BlueprintPure, Category = "Utilities|Time Management")
-	static FFrameRate GetTimecodeFrameRate();
-};
\ No newline at end of file
+	static TIMEMANAGEMENT_API FFrameRate GetTimecodeFrameRate();
+};
diff --git a/Engine/Source/Runtime/TimeManagement/Public/TimeSynchronizationSource.h b/Engine/Source/Runtime/TimeManagement/Public/TimeSynchronizationSource.h
--- a/Engine/Source/Runtime/TimeManagement/Public/TimeSynchronizationSource.h
+++ b/Engine/Source/Runtime/TimeManagement/Public/TimeSynchronizationSource.h
@@ -49,8 +49,8 @@ struct FTimeSynchronizationStartData
  *
  * Currently, Synchronization does not work on the subframe level.
  */
-UCLASS(Abstract)
-class TIMEMANAGEMENT_API UTimeSynchronizationSource : public UObject
+UCLASS(Abstract, MinimalAPI)
+class UTimeSynchronizationSource : public UObject
 {
 	GENERATED_UCLASS_BODY()
 
@@ -72,38 +72,38 @@ public:
 
 #if WITH_EDITOR
 	/** Get Visual Widget of this source to display in UI */
-	virtual TSharedRef<SWidget> GetVisualWidget() const;
+	TIMEMANAGEMENT_API virtual TSharedRef<SWidget> GetVisualWidget() const;
 #endif
 
 	/**
 	 * Get the time of the newest available sample (relative to this source's frame rate).
 	 * Note, in cases where Rollover is allowed and has occurred, this may have a lower value than GetOldestSampleTime. 
 	 */
-	virtual FFrameTime GetNewestSampleTime() const PURE_VIRTUAL(UTimeSynchronizationSource::GetNewestSampleTime, return FFrameTime();)
+	TIMEMANAGEMENT_API virtual FFrameTime GetNewestSampleTime() const PURE_VIRTUAL(UTimeSynchronizationSource::GetNewestSampleTime, return FFrameTime();)
 
 	/**
 	 * Get the time of the oldest available sample (relative to this source's frame rate).
 	 * Note, in cases where Rollover is allowed and has occurred, this may have a higher value than GetNewestSampleTime.
 	 */
-	virtual FFrameTime GetOldestSampleTime() const PURE_VIRTUAL(UTimeSynchronizationSource::GetOldestSampleTime, return FFrameTime();)
+	TIMEMANAGEMENT_API virtual FFrameTime GetOldestSampleTime() const PURE_VIRTUAL(UTimeSynchronizationSource::GetOldestSampleTime, return FFrameTime();)
 
 	/** Get the source actual FrameRate */
-	virtual FFrameRate GetFrameRate() const PURE_VIRTUAL(UTimeSynchronizationSource::GetFrameRate, return FFrameRate();)
+	TIMEMANAGEMENT_API virtual FFrameRate GetFrameRate() const PURE_VIRTUAL(UTimeSynchronizationSource::GetFrameRate, return FFrameRate();)
 
 	/** Used to know if the source is ready to be used for synchronization. */
-	virtual bool IsReady() const PURE_VIRTUAL(UTimeSynchronizationSource::IsReady, return false;)
+	TIMEMANAGEMENT_API virtual bool IsReady() const PURE_VIRTUAL(UTimeSynchronizationSource::IsReady, return false;)
 
 	/** Called when synchronization is started to notify this source to begin buffering frames. */
-	virtual bool Open(const FTimeSynchronizationOpenData& OpenData) PURE_VIRTUAL(UTimeSynchronizationSource::Open, return false;)
+	TIMEMANAGEMENT_API virtual bool Open(const FTimeSynchronizationOpenData& OpenData) PURE_VIRTUAL(UTimeSynchronizationSource::Open, return false;)
 
 	/** Start playing samples. */
-	virtual void Start(const FTimeSynchronizationStartData& StartData) PURE_VIRTUAL(UTimeSynchronizationSource::Start, return;)
+	TIMEMANAGEMENT_API virtual void Start(const FTimeSynchronizationStartData& StartData) PURE_VIRTUAL(UTimeSynchronizationSource::Start, return;)
 
 	/** Called when synchronization has been completed. The source may discard any unnecessary frames. */
-	virtual void Close() PURE_VIRTUAL(UTimeSynchronizationSource::Close, return;)
+	TIMEMANAGEMENT_API virtual void Close() PURE_VIRTUAL(UTimeSynchronizationSource::Close, return;)
 
 	/** Name to used when displaying an error message or to used in UI. */
-	virtual FString GetDisplayName() const PURE_VIRTUAL(UTimeSynchronizationSource::GetDisplayName, return FString();)
+	TIMEMANAGEMENT_API virtual FString GetDisplayName() const PURE_VIRTUAL(UTimeSynchronizationSource::GetDisplayName, return FString();)
 
 public:
 
diff --git a/Engine/Source/Runtime/TimeManagement/Public/TimedDataInputCollection.h b/Engine/Source/Runtime/TimeManagement/Public/TimedDataInputCollection.h
--- a/Engine/Source/Runtime/TimeManagement/Public/TimedDataInputCollection.h
+++ b/Engine/Source/Runtime/TimeManagement/Public/TimedDataInputCollection.h
@@ -11,7 +11,7 @@ class ITimedDataInputChannel;
 /**
  * A list of all the timed data input.
  */
-class TIMEMANAGEMENT_API FTimedDataInputCollection
+class FTimedDataInputCollection
 {
 public:
 
@@ -75,4 +75,4 @@ private:
 	FSimpleMulticastDelegate CollectionChanged;
 	TArray<ITimedDataInput*> Inputs;
 	TArray<ITimedDataInputChannel*> Channels;
-};
\ No newline at end of file
+};
diff --git a/Engine/Source/Runtime/TimeManagement/Public/Widgets/STimecode.h b/Engine/Source/Runtime/TimeManagement/Public/Widgets/STimecode.h
--- a/Engine/Source/Runtime/TimeManagement/Public/Widgets/STimecode.h
+++ b/Engine/Source/Runtime/TimeManagement/Public/Widgets/STimecode.h
@@ -12,7 +12,7 @@
 #include "Styling/CoreStyle.h"
 #include "Styling/SlateColor.h"
 
-class TIMEMANAGEMENT_API STimecode : public SLeafWidget
+class STimecode : public SLeafWidget
 {
 public:
 	SLATE_BEGIN_ARGS(STimecode)
@@ -40,20 +40,20 @@ public:
 		SLATE_ATTRIBUTE(FSlateColor, LabelColor)
 	SLATE_END_ARGS()
 
-	STimecode();
+	TIMEMANAGEMENT_API STimecode();
 
 	/**
 	 * Construct this widget
 	 *
 	 * @param	InArgs	The declaration data for this widget
 	 */
-	void Construct(const FArguments& InArgs);
+	TIMEMANAGEMENT_API void Construct(const FArguments& InArgs);
 
 protected:
 	// SWidget overrides
-	virtual int32 OnPaint(const FPaintArgs& Args, const FGeometry& AllottedGeometry, const FSlateRect& MyCullingRect, FSlateWindowElementList& OutDrawElements, int32 LayerId, const FWidgetStyle& InWidgetStyle, bool bParentEnabled) const override;
-	virtual FVector2D ComputeDesiredSize(float) const override;
-	virtual bool ComputeVolatility() const override;
+	TIMEMANAGEMENT_API virtual int32 OnPaint(const FPaintArgs& Args, const FGeometry& AllottedGeometry, const FSlateRect& MyCullingRect, FSlateWindowElementList& OutDrawElements, int32 LayerId, const FWidgetStyle& InWidgetStyle, bool bParentEnabled) const override;
+	TIMEMANAGEMENT_API virtual FVector2D ComputeDesiredSize(float) const override;
+	TIMEMANAGEMENT_API virtual bool ComputeVolatility() const override;
 
 private:
 	TAttribute<FTimecode> Timecode;
diff --git a/Engine/Source/Runtime/TypedElementFramework/Public/Elements/Framework/TypedElementCounter.h b/Engine/Source/Runtime/TypedElementFramework/Public/Elements/Framework/TypedElementCounter.h
--- a/Engine/Source/Runtime/TypedElementFramework/Public/Elements/Framework/TypedElementCounter.h
+++ b/Engine/Source/Runtime/TypedElementFramework/Public/Elements/Framework/TypedElementCounter.h
@@ -46,13 +46,13 @@ inline FName GetTypedElementCounterKeyName()
  * Tracks various categories of counters for the typed elements (eg, the number of elements of a given type).
  * The categories counted may be expanded by element types that implement the UTypedElementCounterInterface API.
  */
-class TYPEDELEMENTFRAMEWORK_API FTypedElementCounter
+class FTypedElementCounter
 {
 public:
 	using FCounterValue = int32;
 
 	FTypedElementCounter() = default;
-	explicit FTypedElementCounter(UTypedElementRegistry* InRegistry);
+	TYPEDELEMENTFRAMEWORK_API explicit FTypedElementCounter(UTypedElementRegistry* InRegistry);
 
 	FTypedElementCounter(const FTypedElementCounter&) = delete;
 	FTypedElementCounter& operator=(const FTypedElementCounter&) = delete;
@@ -63,17 +63,17 @@ public:
 	/**
 	 * Explicitly initialize this instance, if it was previously default constructed.
 	 */
-	void Initialize(UTypedElementRegistry* InRegistry);
+	TYPEDELEMENTFRAMEWORK_API void Initialize(UTypedElementRegistry* InRegistry);
 
 	/**
 	 * Update the counter categories to include the given element.
 	 */
-	void AddElement(const FTypedElementHandle& InElementHandle);
+	TYPEDELEMENTFRAMEWORK_API void AddElement(const FTypedElementHandle& InElementHandle);
 
 	/**
 	 * Update the counter categories to remove the given element.
 	 */
-	void RemoveElement(const FTypedElementHandle& InElementHandle);
+	TYPEDELEMENTFRAMEWORK_API void RemoveElement(const FTypedElementHandle& InElementHandle);
 
 	/**
 	 * Increment the counter for the given category and key.
@@ -143,17 +143,17 @@ public:
 	/**
 	 * Clear the value of the counters for the given category.
 	 */
-	void ClearCounters(const FName InCategory);
+	TYPEDELEMENTFRAMEWORK_API void ClearCounters(const FName InCategory);
 
 	/**
 	 * Clear all counters.
 	 */
-	void ClearCounters();
+	TYPEDELEMENTFRAMEWORK_API void ClearCounters();
 
 	/**
 	 * Get the category name used to count element types (by type ID).
 	 */
-	static FName GetElementTypeCategoryName();
+	static TYPEDELEMENTFRAMEWORK_API FName GetElementTypeCategoryName();
 
 private:
 	class ICounterCategory
@@ -296,7 +296,7 @@ class UTypedElementCounterInterface : public UInterface
 	GENERATED_BODY()
 };
 
-class TYPEDELEMENTFRAMEWORK_API ITypedElementCounterInterface
+class ITypedElementCounterInterface
 {
 	GENERATED_BODY()
 
diff --git a/Engine/Source/Runtime/TypedElementFramework/Public/Elements/Framework/TypedElementHandle.h b/Engine/Source/Runtime/TypedElementFramework/Public/Elements/Framework/TypedElementHandle.h
--- a/Engine/Source/Runtime/TypedElementFramework/Public/Elements/Framework/TypedElementHandle.h
+++ b/Engine/Source/Runtime/TypedElementFramework/Public/Elements/Framework/TypedElementHandle.h
@@ -14,7 +14,7 @@
  * C++ code may choose to use TTypedElement instead, which is a combination of an element handle and its associated element interface.
  * @note Handles auto-release on destruction.
  */
-struct TYPEDELEMENTFRAMEWORK_API FTypedElementHandle
+struct FTypedElementHandle
 {
 public:
 	FTypedElementHandle() = default;
@@ -621,7 +621,7 @@ using FTypedElementOwner = TTypedElementOwner<void>;
  * Note: This type shouldn't be used in the engine code as it come with a performance and memory overhead that we want to avoid when compare to the native handles (FTypedElementHandle).
  */
 USTRUCT(BlueprintType)
-struct TYPEDELEMENTFRAMEWORK_API FScriptTypedElementHandle
+struct FScriptTypedElementHandle
 {
 	GENERATED_BODY()
 public:
diff --git a/Engine/Source/Runtime/TypedElementFramework/Public/Elements/Framework/TypedElementId.h b/Engine/Source/Runtime/TypedElementFramework/Public/Elements/Framework/TypedElementId.h
--- a/Engine/Source/Runtime/TypedElementFramework/Public/Elements/Framework/TypedElementId.h
+++ b/Engine/Source/Runtime/TypedElementFramework/Public/Elements/Framework/TypedElementId.h
@@ -13,7 +13,7 @@
  * @note This is ref-counted like handles themselves are, so as long as an ID is available, the handle will be too.
  * @note IDs lack the information needed to auto-release on destruction, so must be manually released, either via the corresponding handle or their owner element registry.
  */
-struct TYPEDELEMENTFRAMEWORK_API FTypedElementId
+struct FTypedElementId
 {
 public:
 	FTypedElementId()
@@ -113,7 +113,7 @@ public:
 	}
 
 	/** An unset element ID */
-	static const FTypedElementId Unset;
+	static TYPEDELEMENTFRAMEWORK_API const FTypedElementId Unset;
 
 private:
 	union
diff --git a/Engine/Source/Runtime/TypedElementFramework/Public/Elements/Framework/TypedElementList.h b/Engine/Source/Runtime/TypedElementFramework/Public/Elements/Framework/TypedElementList.h
--- a/Engine/Source/Runtime/TypedElementFramework/Public/Elements/Framework/TypedElementList.h
+++ b/Engine/Source/Runtime/TypedElementFramework/Public/Elements/Framework/TypedElementList.h
@@ -58,7 +58,7 @@ FORCEINLINE TTypedElement<BaseInterfaceType> GetElement(const UTypedElementRegis
  * Provides high-level access to groups of elements, including accessing elements that implement specific interfaces.
  */
 template<class HandleType>
-class TYPEDELEMENTFRAMEWORK_API TTypedElementList final : public TSharedFromThis<TTypedElementList<HandleType>>
+class TTypedElementList final : public TSharedFromThis<TTypedElementList<HandleType>>
 {
 public:
 
@@ -564,18 +564,18 @@ public:
 	 * A utility struct that help to cancel any new pending notification that happened in a scope.
 	 * Note: it won't cancel a notification if there is a legacy batch operation ongoing
 	 */
-	struct TYPEDELEMENTFRAMEWORK_API FScopedClearNewPendingChange
+	struct FScopedClearNewPendingChange
 	{
 		FScopedClearNewPendingChange() = default;
-		FScopedClearNewPendingChange(TTypedElementList& InTypeElementList);
+		TYPEDELEMENTFRAMEWORK_API FScopedClearNewPendingChange(TTypedElementList& InTypeElementList);
 
 		FScopedClearNewPendingChange(const FScopedClearNewPendingChange&) = delete;
 		FScopedClearNewPendingChange& operator=(const FScopedClearNewPendingChange&) = delete;
 
-		FScopedClearNewPendingChange(FScopedClearNewPendingChange&& Other);
-		FScopedClearNewPendingChange& operator=(FScopedClearNewPendingChange&& Other);
+		TYPEDELEMENTFRAMEWORK_API FScopedClearNewPendingChange(FScopedClearNewPendingChange&& Other);
+		TYPEDELEMENTFRAMEWORK_API FScopedClearNewPendingChange& operator=(FScopedClearNewPendingChange&& Other);
 
-		~FScopedClearNewPendingChange();
+		TYPEDELEMENTFRAMEWORK_API ~FScopedClearNewPendingChange();
 
 	private:
 		TTypedElementList* TypedElementList = nullptr;
@@ -585,7 +585,7 @@ public:
 	 * Interface to allow external systems (such as USelection) to receive immediate sync notifications as an element list is changed.
 	 * This exists purely as a bridging mechanism and shouldn't be relied on for new code. It is lazily created as needed.
 	 */
-	class TYPEDELEMENTFRAMEWORK_API FLegacySync
+	class FLegacySync
 	{
 	public:
 		enum class ESyncType : uint8
@@ -623,18 +623,18 @@ public:
 			BatchComplete,
 		};
 	
-		FLegacySync(const TTypedElementList& InElementList);
+		TYPEDELEMENTFRAMEWORK_API FLegacySync(const TTypedElementList& InElementList);
 
-		void Private_EmitSyncEvent(const ESyncType InSyncType, const HandleType& InElementHandle = HandleType());
+		TYPEDELEMENTFRAMEWORK_API void Private_EmitSyncEvent(const ESyncType InSyncType, const HandleType& InElementHandle = HandleType());
 
 		DECLARE_EVENT_FourParams(FLegacySync, FOnSyncEvent, const TTypedElementList& /*InElementList*/, ESyncType /*InSyncType*/, const HandleType& /*InElementHandle*/, bool /*bIsWithinBatchOperation*/);
-		FOnSyncEvent& OnSyncEvent();
+		TYPEDELEMENTFRAMEWORK_API FOnSyncEvent& OnSyncEvent();
 
-		bool IsRunningBatchOperation() const;
-		void BeginBatchOperation();
-		void EndBatchOperation(const bool InNotify = true);
-		bool IsBatchOperationDirty() const;
-		void ForceBatchOperationDirty();
+		TYPEDELEMENTFRAMEWORK_API bool IsRunningBatchOperation() const;
+		TYPEDELEMENTFRAMEWORK_API void BeginBatchOperation();
+		TYPEDELEMENTFRAMEWORK_API void EndBatchOperation(const bool InNotify = true);
+		TYPEDELEMENTFRAMEWORK_API bool IsBatchOperationDirty() const;
+		TYPEDELEMENTFRAMEWORK_API void ForceBatchOperationDirty();
 
 	private:
 		const TTypedElementList& ElementList;
@@ -649,11 +649,11 @@ public:
 	 * Helper to batch immediate sync notifications for legacy code.
 	 * Does nothing if no legacy sync has been created for the given instance.
 	 */
-	class TYPEDELEMENTFRAMEWORK_API FLegacySyncScopedBatch
+	class FLegacySyncScopedBatch
 	{
 	public:
-		explicit FLegacySyncScopedBatch(const TTypedElementList& InElementList, const bool InNotify = true);
-		~FLegacySyncScopedBatch();
+		TYPEDELEMENTFRAMEWORK_API explicit FLegacySyncScopedBatch(const TTypedElementList& InElementList, const bool InNotify = true);
+		TYPEDELEMENTFRAMEWORK_API ~FLegacySyncScopedBatch();
 
 		FLegacySyncScopedBatch(const FLegacySyncScopedBatch&) = delete;
 		FLegacySyncScopedBatch& operator=(const FLegacySyncScopedBatch&) = delete;
@@ -661,8 +661,8 @@ public:
 		FLegacySyncScopedBatch(FLegacySyncScopedBatch&&) = delete;
 		FLegacySyncScopedBatch& operator=(FLegacySyncScopedBatch&&) = delete;
 
-		bool IsDirty() const;
-		void ForceDirty();
+		TYPEDELEMENTFRAMEWORK_API bool IsDirty() const;
+		TYPEDELEMENTFRAMEWORK_API void ForceDirty();
 
 	private:
 		FLegacySync* ElementListLegacySync = nullptr;
@@ -787,4 +787,4 @@ namespace UE::TypedElementFramework
 	 */
 	TYPEDELEMENTFRAMEWORK_API FScriptTypedElementListPtr ConvertToScriptTypedElementList(const FTypedElementListConstPtr& NativeList);
 
-}
\ No newline at end of file
+}
diff --git a/Engine/Source/Runtime/TypedElementFramework/Public/Elements/Framework/TypedElementListProxy.h b/Engine/Source/Runtime/TypedElementFramework/Public/Elements/Framework/TypedElementListProxy.h
--- a/Engine/Source/Runtime/TypedElementFramework/Public/Elements/Framework/TypedElementListProxy.h
+++ b/Engine/Source/Runtime/TypedElementFramework/Public/Elements/Framework/TypedElementListProxy.h
@@ -14,7 +14,7 @@
  * The weak model for the handles come with an additional cost to the runtime performance and the memory usage.
  */
 USTRUCT(BlueprintType, DisplayName="Typed Element List", meta=(ScriptName="TypedElementList"))
-struct TYPEDELEMENTFRAMEWORK_API FScriptTypedElementListProxy
+struct FScriptTypedElementListProxy
 {
 	GENERATED_BODY()
 
diff --git a/Engine/Source/Runtime/TypedElementFramework/Public/Elements/Framework/TypedElementQueryBuilder.h b/Engine/Source/Runtime/TypedElementFramework/Public/Elements/Framework/TypedElementQueryBuilder.h
--- a/Engine/Source/Runtime/TypedElementFramework/Public/Elements/Framework/TypedElementQueryBuilder.h
+++ b/Engine/Source/Runtime/TypedElementFramework/Public/Elements/Framework/TypedElementQueryBuilder.h
@@ -109,7 +109,7 @@ namespace TypedElementQueryBuilder
 	TYPEDELEMENTFRAMEWORK_API const UScriptStruct* operator""_Type(const char* Name, std::size_t NameSize);
 	TYPEDELEMENTFRAMEWORK_API const UScriptStruct* operator""_TypeOptional(const char* Name, std::size_t NameSize);
 
-	class TYPEDELEMENTFRAMEWORK_API FDependency final
+	class FDependency final
 	{
 		friend class Count;
 		friend class Select;
@@ -117,62 +117,62 @@ namespace TypedElementQueryBuilder
 	public:
 		template<typename... TargetTypes>
 		FDependency& ReadOnly();
-		FDependency& ReadOnly(const UClass* Target);
-		FDependency& ReadOnly(TConstArrayView<const UClass*> Targets);
+		TYPEDELEMENTFRAMEWORK_API FDependency& ReadOnly(const UClass* Target);
+		TYPEDELEMENTFRAMEWORK_API FDependency& ReadOnly(TConstArrayView<const UClass*> Targets);
 		template<typename... TargetTypes>
 		FDependency& ReadWrite();
-		FDependency& ReadWrite(const UClass* Target);
-		FDependency& ReadWrite(TConstArrayView<const UClass*> Targets);
+		TYPEDELEMENTFRAMEWORK_API FDependency& ReadWrite(const UClass* Target);
+		TYPEDELEMENTFRAMEWORK_API FDependency& ReadWrite(TConstArrayView<const UClass*> Targets);
 
-		FDependency& SubQuery(TypedElementQueryHandle Handle);
-		FDependency& SubQuery(TConstArrayView<TypedElementQueryHandle> Handles);
+		TYPEDELEMENTFRAMEWORK_API FDependency& SubQuery(TypedElementQueryHandle Handle);
+		TYPEDELEMENTFRAMEWORK_API FDependency& SubQuery(TConstArrayView<TypedElementQueryHandle> Handles);
 
-		ITypedElementDataStorageInterface::FQueryDescription&& Compile();
+		TYPEDELEMENTFRAMEWORK_API ITypedElementDataStorageInterface::FQueryDescription&& Compile();
 
 	private:
-		explicit FDependency(ITypedElementDataStorageInterface::FQueryDescription* Query);
+		TYPEDELEMENTFRAMEWORK_API explicit FDependency(ITypedElementDataStorageInterface::FQueryDescription* Query);
 
 		ITypedElementDataStorageInterface::FQueryDescription* Query;
 	};
 
-	class TYPEDELEMENTFRAMEWORK_API FSimpleQuery final
+	class FSimpleQuery final
 	{
 	public:
 		friend class Count;
 		friend class Select;
 
-		FDependency DependsOn();
-		ITypedElementDataStorageInterface::FQueryDescription&& Compile();
+		TYPEDELEMENTFRAMEWORK_API FDependency DependsOn();
+		TYPEDELEMENTFRAMEWORK_API ITypedElementDataStorageInterface::FQueryDescription&& Compile();
 
 		template<typename... TargetTypes>
 		FSimpleQuery& All();
-		FSimpleQuery& All(const UScriptStruct* Target);
-		FSimpleQuery& All(TConstArrayView<const UScriptStruct*> Targets);
+		TYPEDELEMENTFRAMEWORK_API FSimpleQuery& All(const UScriptStruct* Target);
+		TYPEDELEMENTFRAMEWORK_API FSimpleQuery& All(TConstArrayView<const UScriptStruct*> Targets);
 		template<typename... TargetTypes>
 		FSimpleQuery& Any();
-		FSimpleQuery& Any(const UScriptStruct* Target);
-		FSimpleQuery& Any(TConstArrayView<const UScriptStruct*> Targets);
+		TYPEDELEMENTFRAMEWORK_API FSimpleQuery& Any(const UScriptStruct* Target);
+		TYPEDELEMENTFRAMEWORK_API FSimpleQuery& Any(TConstArrayView<const UScriptStruct*> Targets);
 		template<typename... TargetTypes>
 		FSimpleQuery& None();
-		FSimpleQuery& None(const UScriptStruct* Target);
-		FSimpleQuery& None(TConstArrayView<const UScriptStruct*> Targets);
+		TYPEDELEMENTFRAMEWORK_API FSimpleQuery& None(const UScriptStruct* Target);
+		TYPEDELEMENTFRAMEWORK_API FSimpleQuery& None(TConstArrayView<const UScriptStruct*> Targets);
 
 	private:
-		explicit FSimpleQuery(ITypedElementDataStorageInterface::FQueryDescription* Query);
+		TYPEDELEMENTFRAMEWORK_API explicit FSimpleQuery(ITypedElementDataStorageInterface::FQueryDescription* Query);
 
 		ITypedElementDataStorageInterface::FQueryDescription* Query;
 	};
 
-	struct TYPEDELEMENTFRAMEWORK_API FQueryCallbackType{};
+	struct FQueryCallbackType{};
 
-	struct TYPEDELEMENTFRAMEWORK_API FProcessor final : public FQueryCallbackType
+	struct FProcessor final : public FQueryCallbackType
 	{
-		FProcessor(ITypedElementDataStorageInterface::EQueryTickPhase Phase, FName Group);
-		FProcessor& SetPhase(ITypedElementDataStorageInterface::EQueryTickPhase NewPhase);
-		FProcessor& SetGroup(FName GroupName);
-		FProcessor& SetBeforeGroup(FName GroupName);
-		FProcessor& SetAfterGroup(FName GroupName);
-		FProcessor& ForceToGameThread(bool bForce);
+		TYPEDELEMENTFRAMEWORK_API FProcessor(ITypedElementDataStorageInterface::EQueryTickPhase Phase, FName Group);
+		TYPEDELEMENTFRAMEWORK_API FProcessor& SetPhase(ITypedElementDataStorageInterface::EQueryTickPhase NewPhase);
+		TYPEDELEMENTFRAMEWORK_API FProcessor& SetGroup(FName GroupName);
+		TYPEDELEMENTFRAMEWORK_API FProcessor& SetBeforeGroup(FName GroupName);
+		TYPEDELEMENTFRAMEWORK_API FProcessor& SetAfterGroup(FName GroupName);
+		TYPEDELEMENTFRAMEWORK_API FProcessor& ForceToGameThread(bool bForce);
 		
 		ITypedElementDataStorageInterface::EQueryTickPhase Phase;
 		FName Group;
@@ -181,7 +181,7 @@ namespace TypedElementQueryBuilder
 		bool bForceToGameThread{ false };
 	};
 
-	struct TYPEDELEMENTFRAMEWORK_API FObserver final : public FQueryCallbackType
+	struct FObserver final : public FQueryCallbackType
 	{
 		enum class EEvent : uint8
 		{
@@ -191,20 +191,20 @@ namespace TypedElementQueryBuilder
 
 		template<typename ColumnType>
 		explicit FObserver(EEvent MonitorForEvent);
-		FObserver(EEvent MonitorForEvent, const UScriptStruct* MonitoredColumn);
+		TYPEDELEMENTFRAMEWORK_API FObserver(EEvent MonitorForEvent, const UScriptStruct* MonitoredColumn);
 
-		FObserver& SetEvent(EEvent MonitorForEvent);
-		FObserver& SetMonitoredColumn(const UScriptStruct* MonitoredColumn);
+		TYPEDELEMENTFRAMEWORK_API FObserver& SetEvent(EEvent MonitorForEvent);
+		TYPEDELEMENTFRAMEWORK_API FObserver& SetMonitoredColumn(const UScriptStruct* MonitoredColumn);
 		template<typename ColumnType>
 		FObserver& SetMonitoredColumn();
-		FObserver& ForceToGameThread(bool bForce);
+		TYPEDELEMENTFRAMEWORK_API FObserver& ForceToGameThread(bool bForce);
 
 		const UScriptStruct* Monitor;
 		EEvent Event;
 		bool bForceToGameThread{ false };
 	};
 
-	struct TYPEDELEMENTFRAMEWORK_API FPhaseAmble final : public FQueryCallbackType
+	struct FPhaseAmble final : public FQueryCallbackType
 	{
 		enum class ELocation : uint8
 		{
@@ -212,10 +212,10 @@ namespace TypedElementQueryBuilder
 			Postamble
 		};
 
-		FPhaseAmble(ELocation InLocation, ITypedElementDataStorageInterface::EQueryTickPhase InPhase);
-		FPhaseAmble& SetLocation(ELocation NewLocation);
-		FPhaseAmble& SetPhase(ITypedElementDataStorageInterface::EQueryTickPhase NewPhase);
-		FPhaseAmble& ForceToGameThread(bool bForce);
+		TYPEDELEMENTFRAMEWORK_API FPhaseAmble(ELocation InLocation, ITypedElementDataStorageInterface::EQueryTickPhase InPhase);
+		TYPEDELEMENTFRAMEWORK_API FPhaseAmble& SetLocation(ELocation NewLocation);
+		TYPEDELEMENTFRAMEWORK_API FPhaseAmble& SetPhase(ITypedElementDataStorageInterface::EQueryTickPhase NewPhase);
+		TYPEDELEMENTFRAMEWORK_API FPhaseAmble& ForceToGameThread(bool bForce);
 
 		ITypedElementDataStorageInterface::EQueryTickPhase Phase;
 		ELocation Location;
@@ -276,10 +276,10 @@ namespace TypedElementQueryBuilder
 	};
 
 	// Explicitly not following the naming convention in order to present this as a query that can be read as such.
-	class TYPEDELEMENTFRAMEWORK_API Select final
+	class Select final
 	{
 	public:
-		Select();
+		TYPEDELEMENTFRAMEWORK_API Select();
 
 		template<typename CallbackType, typename Function>
 		Select(FName Name, const CallbackType& Type, Function&& Callback);
@@ -288,29 +288,29 @@ namespace TypedElementQueryBuilder
 
 		template<typename... TargetTypes>
 		Select& ReadOnly();
-		Select& ReadOnly(const UScriptStruct* Target);
-		Select& ReadOnly(std::initializer_list<const UScriptStruct*> Targets);
+		TYPEDELEMENTFRAMEWORK_API Select& ReadOnly(const UScriptStruct* Target);
+		TYPEDELEMENTFRAMEWORK_API Select& ReadOnly(std::initializer_list<const UScriptStruct*> Targets);
 		template<typename... TargetTypes>
 		Select& ReadWrite();
-		Select& ReadWrite(const UScriptStruct* Target);
-		Select& ReadWrite(std::initializer_list<const UScriptStruct*> Targets);
+		TYPEDELEMENTFRAMEWORK_API Select& ReadWrite(const UScriptStruct* Target);
+		TYPEDELEMENTFRAMEWORK_API Select& ReadWrite(std::initializer_list<const UScriptStruct*> Targets);
 
-		ITypedElementDataStorageInterface::FQueryDescription&& Compile();
-		FSimpleQuery Where();
-		FDependency DependsOn();
+		TYPEDELEMENTFRAMEWORK_API ITypedElementDataStorageInterface::FQueryDescription&& Compile();
+		TYPEDELEMENTFRAMEWORK_API FSimpleQuery Where();
+		TYPEDELEMENTFRAMEWORK_API FDependency DependsOn();
 
 	private:
 		ITypedElementDataStorageInterface::FQueryDescription Query;
 	};
 
 	// Explicitly not following the naming convention in order to keep readability consistent. It now reads like a query sentence.
-	class TYPEDELEMENTFRAMEWORK_API Count final
+	class Count final
 	{
 	public:
-		Count();
+		TYPEDELEMENTFRAMEWORK_API Count();
 
-		FSimpleQuery Where();
-		FDependency DependsOn();
+		TYPEDELEMENTFRAMEWORK_API FSimpleQuery Where();
+		TYPEDELEMENTFRAMEWORK_API FDependency DependsOn();
 
 	private:
 		ITypedElementDataStorageInterface::FQueryDescription Query;
