NegativeScale automation test fixes (ApplyTransform/AppenMeshDescription fixes)

diff --git a/Engine/Source/Runtime/StaticMeshDescription/Private/StaticMeshOperations.cpp b/Engine/Source/Runtime/StaticMeshDescription/Private/StaticMeshOperations.cpp
--- a/Engine/Source/Runtime/StaticMeshDescription/Private/StaticMeshOperations.cpp
+++ b/Engine/Source/Runtime/StaticMeshDescription/Private/StaticMeshOperations.cpp
@@ -1183,7 +1183,8 @@ void FStaticMeshOperations::AppendMeshDescriptions(const TArray<const FMeshDescr
 		if (AppendSettings.MeshTransform)
 		{
 			const FTransform& Transform = AppendSettings.MeshTransform.GetValue();
-			const FMatrix TransformInverseTransposeMatrix = Transform.ToMatrixWithScale().Inverse().GetTransposed();
+			FMatrix TransformInverseTransposeMatrix = Transform.ToMatrixWithScale().Inverse().GetTransposed();
+			TransformInverseTransposeMatrix.RemoveScaling();
 
 			bReverseCulling = Transform.GetDeterminant() < 0;
 			float BinormalSignsFactor = bReverseCulling ? -1.f : 1.f;
@@ -2654,6 +2655,7 @@ void FStaticMeshOperations::ApplyTransform(FMeshDescription& MeshDescription, co
 	}
 
 	FMatrix TransformInverseTransposeMatrix = Transform.Inverse().GetTransposed();
+	TransformInverseTransposeMatrix.RemoveScaling();
 
 	const bool bIsMirrored = Transform.Determinant() < 0.f;
 	const float MulBy = bIsMirrored ? -1.f : 1.f;
