Downgrade NAN meshdescription log from Error to Display level.

#rnx

diff --git a/Engine/Source/Runtime/StaticMeshDescription/Private/StaticMeshOperations.cpp b/Engine/Source/Runtime/StaticMeshDescription/Private/StaticMeshOperations.cpp
--- a/Engine/Source/Runtime/StaticMeshDescription/Private/StaticMeshOperations.cpp
+++ b/Engine/Source/Runtime/StaticMeshDescription/Private/StaticMeshOperations.cpp
@@ -2999,19 +2999,19 @@ bool FStaticMeshOperations::ValidateAndFixData(FMeshDescription& MeshDescription
 
 	if (bHasInvalidPositions)
 	{
-		UE_LOG(LogStaticMeshOperations, Error, TEXT("Mesh %s has NaNs in it's vertex positions! Offending positions are set to zero."), *DebugName);
+		UE_LOG(LogStaticMeshOperations, Display, TEXT("Mesh %s has NaNs in it's vertex positions! Offending positions are set to zero."), *DebugName);
 	}
 	if (bHasInvalidTangentSpaces)
 	{
-		UE_LOG(LogStaticMeshOperations, Error, TEXT("Mesh %s has NaNs in it's vertex instance tangent space! Offending tangents are set to zero."), *DebugName);
+		UE_LOG(LogStaticMeshOperations, Display, TEXT("Mesh %s has NaNs in it's vertex instance tangent space! Offending tangents are set to zero."), *DebugName);
 	}
 	if (bHasInvalidUVs)
 	{
-		UE_LOG(LogStaticMeshOperations, Error, TEXT("Mesh %s has NaNs in it's vertex instance uvs! Offending uvs are set to zero."), *DebugName);
+		UE_LOG(LogStaticMeshOperations, Display, TEXT("Mesh %s has NaNs in it's vertex instance uvs! Offending uvs are set to zero."), *DebugName);
 	}
 	if (bHasInvalidVertexColors)
 	{
-		UE_LOG(LogStaticMeshOperations, Error, TEXT("Mesh %s has NaNs in it's vertex instance colors! Offending colors are set to white."), *DebugName);
+		UE_LOG(LogStaticMeshOperations, Display, TEXT("Mesh %s has NaNs in it's vertex instance colors! Offending colors are set to white."), *DebugName);
 	}
 
 	return !bHasInvalidPositions && !bHasInvalidTangentSpaces && !bHasInvalidUVs && !bHasInvalidVertexColors;
