MergeActor - Improved UV generation of merged/simplified meshes

Currently controlled globally with the MeshMergeUtilities.UVGenerationMethod cvar
-> 0 - Engine default - (currently Patch Builder)
-> 1 - Legacy
-> 2 - UVAtlas
-> 3 - XAtlas
-> 4 - Patch Builder

diff --git a/Engine/Plugins/Editor/WorldPartitionHLODUtilities/Source/Private/WorldPartition/HLOD/Builders/HLODBuilderMeshMerge.cpp b/Engine/Plugins/Editor/WorldPartitionHLODUtilities/Source/Private/WorldPartition/HLOD/Builders/HLODBuilderMeshMerge.cpp
--- a/Engine/Plugins/Editor/WorldPartitionHLODUtilities/Source/Private/WorldPartition/HLOD/Builders/HLODBuilderMeshMerge.cpp
+++ b/Engine/Plugins/Editor/WorldPartitionHLODUtilities/Source/Private/WorldPartition/HLOD/Builders/HLODBuilderMeshMerge.cpp
@@ -54,6 +54,10 @@ uint32 UHLODBuilderMeshMergeSettings::GetCRC() const
 		IMaterialBakingModule& Module = FModuleManager::Get().LoadModuleChecked<IMaterialBakingModule>("MaterialBaking");
 		uint32 MaterialBakingModuleCRC = Module.GetCRC();
 		Ar << MaterialBakingModuleCRC;
+
+		static const auto MeshMergeUtilitiesUVGenerationMethodCVar = IConsoleManager::Get().FindConsoleVariable(TEXT("MeshMergeUtilities.UVGenerationMethod"));
+		int32 MeshMergeUtilitiesUVGenerationMethod = (MeshMergeUtilitiesUVGenerationMethodCVar != nullptr) ? MeshMergeUtilitiesUVGenerationMethodCVar->GetInt() : 0;
+		Ar << MeshMergeUtilitiesUVGenerationMethod;
 	}
 
 	uint32 Hash = Ar.GetCrc();
diff --git a/Engine/Plugins/Editor/WorldPartitionHLODUtilities/Source/Private/WorldPartition/HLOD/Builders/HLODBuilderMeshSimplify.cpp b/Engine/Plugins/Editor/WorldPartitionHLODUtilities/Source/Private/WorldPartition/HLOD/Builders/HLODBuilderMeshSimplify.cpp
--- a/Engine/Plugins/Editor/WorldPartitionHLODUtilities/Source/Private/WorldPartition/HLOD/Builders/HLODBuilderMeshSimplify.cpp
+++ b/Engine/Plugins/Editor/WorldPartitionHLODUtilities/Source/Private/WorldPartition/HLOD/Builders/HLODBuilderMeshSimplify.cpp
@@ -54,6 +54,10 @@ uint32 UHLODBuilderMeshSimplifySettings::GetCRC() const
 	uint32 MaterialBakingModuleCRC = Module.GetCRC();
 	Ar << MaterialBakingModuleCRC;
 
+	static const auto MeshMergeUtilitiesUVGenerationMethodCVar = IConsoleManager::Get().FindConsoleVariable(TEXT("MeshMergeUtilities.UVGenerationMethod"));
+	int32 MeshMergeUtilitiesUVGenerationMethod = (MeshMergeUtilitiesUVGenerationMethodCVar != nullptr) ? MeshMergeUtilitiesUVGenerationMethodCVar->GetInt() : 0;
+	Ar << MeshMergeUtilitiesUVGenerationMethod;
+
 	uint32 Hash = Ar.GetCrc();
 
 	if (!HLODMaterial.IsNull())
diff --git a/Engine/Source/Developer/MeshMergeUtilities/Private/MeshMergeUtilities.cpp b/Engine/Source/Developer/MeshMergeUtilities/Private/MeshMergeUtilities.cpp
--- a/Engine/Source/Developer/MeshMergeUtilities/Private/MeshMergeUtilities.cpp
+++ b/Engine/Source/Developer/MeshMergeUtilities/Private/MeshMergeUtilities.cpp
@@ -87,6 +87,28 @@
 
 DEFINE_LOG_CATEGORY(LogMeshMerging);
 
+static TAutoConsoleVariable<int32> CVarMeshMergeUtilitiesUVGenerationMethod(
+	TEXT("MeshMergeUtilities.UVGenerationMethod"),
+	0,
+	TEXT("UV generation method when creating merged or proxy meshes\n"
+		 "0 - Engine default - (currently Patch Builder)\n"
+		 "1 - Legacy\n"
+		 "2 - UVAtlas\n"
+		 "3 - XAtlas\n"
+		 "4 - Patch Builder\n"));
+
+static FStaticMeshOperations::EGenerateUVMethod GetUVGenerationMethodToUse()
+{
+	switch (CVarMeshMergeUtilitiesUVGenerationMethod.GetValueOnAnyThread())
+	{
+	case 1:  return FStaticMeshOperations::EGenerateUVMethod::Legacy;
+	case 2:  return FStaticMeshOperations::EGenerateUVMethod::UVAtlas;
+	case 3:  return FStaticMeshOperations::EGenerateUVMethod::XAtlas;
+	case 4:  return FStaticMeshOperations::EGenerateUVMethod::PatchBuilder;
+	default: return FStaticMeshOperations::EGenerateUVMethod::Default;
+	}
+}
+
 FMeshMergeUtilities::FMeshMergeUtilities()
 {
 	Processor = new FProxyGenerationProcessor(this);
@@ -1109,10 +1131,15 @@ static TArray<FVector2D> GetCustomTextureCoordinates(const FMeshDescription& InM
 	}
 	else
 	{
-		bool bSuccess = FStaticMeshOperations::GenerateUniqueUVsForStaticMesh(InMeshDescription, InMeshProxySettings.MaterialSettings.TextureSize.GetMax(), false, CustomTextureCoordinates);
+		FStaticMeshOperations::FGenerateUVOptions GenerateUVOptions;
+		GenerateUVOptions.TextureResolution = InMeshProxySettings.MaterialSettings.TextureSize.GetMax();
+		GenerateUVOptions.bMergeTrianglesWithIdenticalAttributes = false;
+		GenerateUVOptions.UVMethod = GetUVGenerationMethodToUse();
+
+		bool bSuccess = FStaticMeshOperations::GenerateUV(InMeshDescription, GenerateUVOptions, CustomTextureCoordinates);
 		if (!bSuccess)
 		{
-			UE_LOG(LogMeshMerging, Warning, TEXT("GenerateUniqueUVsForStaticMesh: Failed to pack UVs for static mesh \"%s\" (num triangles = %d, texture resolution = %d)."), *InStaticMesh->GetName(), InMeshDescription.Triangles().Num(), InMeshProxySettings.MaterialSettings.TextureSize.GetMax());
+			UE_LOG(LogMeshMerging, Warning, TEXT("GenerateUV: Failed to pack UVs for static mesh \"%s\" (num triangles = %d, texture resolution = %d)."), *InStaticMesh->GetName(), InMeshDescription.Triangles().Num(), InMeshProxySettings.MaterialSettings.TextureSize.GetMax());
 			CustomTextureCoordinates.Empty();
 		}
 	}
@@ -2844,12 +2871,22 @@ void FMeshMergeUtilities::CreateMergedMaterial(FMeshMergeDataTracker& InDataTrac
 							{
 								// No job found yet, fire an async task
 								MeshLODsTextureCoordinates.Add(Tuple, Async(EAsyncExecution::Thread, [MeshDescription, MaterialOptions, this]()
+								{
+									FStaticMeshOperations::FGenerateUVOptions GenerateUVOptions;
+									GenerateUVOptions.TextureResolution = MaterialOptions->TextureSize.GetMax();
+									GenerateUVOptions.bMergeTrianglesWithIdenticalAttributes = false;
+									GenerateUVOptions.UVMethod = GetUVGenerationMethodToUse();
+
+									TArray<FVector2D> UniqueTextureCoordinates;
+									FStaticMeshOperations::GenerateUV(*MeshDescription, GenerateUVOptions, UniqueTextureCoordinates);
+
+									if (GenerateUVOptions.UVMethod == FStaticMeshOperations::EGenerateUVMethod::Legacy)
 									{
-										TArray<FVector2D> UniqueTextureCoordinates;
-										FStaticMeshOperations::GenerateUniqueUVsForStaticMesh(*MeshDescription, MaterialOptions->TextureSize.GetMax(), false, UniqueTextureCoordinates);
 										ScaleTextureCoordinatesToBox(FBox2D(FVector2D::ZeroVector, FVector2D(1, 1)), UniqueTextureCoordinates);
-										return UniqueTextureCoordinates;
-									}));
+									}
+
+									return UniqueTextureCoordinates;
+								}));
 							}
 							// Keep track of the fact that this mesh is waiting for the UV computation to finish
 							MeshDataAwaitingResults.Add(MeshDataIndex, Tuple);
@@ -2895,11 +2932,19 @@ void FMeshMergeUtilities::CreateMergedMaterial(FMeshMergeDataTracker& InDataTrac
 		CreateMergedRawMeshes(InDataTracker, RemapUVMergeSettings, InStaticMeshComponentsToMerge, InUniqueMaterials, InCollapsedMaterialMap, InOutputMaterialsMap, false, false, InMergedAssetPivot, MergedRawMeshes);
 
 		// Create texture coords for the merged mesh
+		FStaticMeshOperations::FGenerateUVOptions GenerateUVOptions;
+		GenerateUVOptions.TextureResolution = MaterialOptions->TextureSize.GetMax();
+		GenerateUVOptions.bMergeTrianglesWithIdenticalAttributes = true;
+		GenerateUVOptions.UVMethod = GetUVGenerationMethodToUse();
+
 		TArray<FVector2D> GlobalTextureCoordinates;
-		bool bSuccess = FStaticMeshOperations::GenerateUniqueUVsForStaticMesh(MergedRawMeshes[0], MaterialOptions->TextureSize.GetMax(), true, GlobalTextureCoordinates);
+		bool bSuccess = FStaticMeshOperations::GenerateUV(MergedRawMeshes[0], GenerateUVOptions, GlobalTextureCoordinates);
 		if (bSuccess)
 		{
-			ScaleTextureCoordinatesToBox(FBox2D(FVector2D::ZeroVector, FVector2D(1, 1)), GlobalTextureCoordinates);
+			if (GenerateUVOptions.UVMethod == FStaticMeshOperations::EGenerateUVMethod::Legacy)
+			{
+				ScaleTextureCoordinatesToBox(FBox2D(FVector2D::ZeroVector, FVector2D(1, 1)), GlobalTextureCoordinates);
+			}
 
 			// copy UVs back to the un-merged mesh's custom texture coords
 			// iterate the raw meshes in the same way as when we combined the mesh above in CreateMergedRawMeshes()
@@ -2935,7 +2980,7 @@ void FMeshMergeUtilities::CreateMergedMaterial(FMeshMergeDataTracker& InDataTrac
 		}
 		else
 		{
-			UE_LOG(LogMeshMerging, Warning, TEXT("GenerateUniqueUVsForStaticMesh: Failed to pack UVs for static mesh"));
+			UE_LOG(LogMeshMerging, Warning, TEXT("GenerateUV: Failed to pack UVs for static mesh"));
 		}
 	}
 
diff --git a/Engine/Source/Runtime/StaticMeshDescription/Private/StaticMeshOperations.cpp b/Engine/Source/Runtime/StaticMeshDescription/Private/StaticMeshOperations.cpp
--- a/Engine/Source/Runtime/StaticMeshDescription/Private/StaticMeshOperations.cpp
+++ b/Engine/Source/Runtime/StaticMeshDescription/Private/StaticMeshOperations.cpp
@@ -1995,96 +1995,148 @@ bool FStaticMeshOperations::CreateLightMapUVLayout(FMeshDescription& MeshDescrip
 	return bPackSuccess;
 }
 
-// Mesh view that will expose only unique UVs if bMergeIdenticalMaterials is provided
-struct FUniqueUVMeshDescriptionView final : FLayoutUV::IMeshView
+static bool GatherUniqueTriangles(const FMeshDescription& InMeshDescription, bool bMergeIdenticalMaterials, TArray<FTriangleID>& OutRemappedTriangles, TArray<FVertexInstanceID>* OutUniqueVerts, TArray<FTriangleID>* OutDuplicateTriangles)
 {
-	const FMeshDescription& MeshDescription;
+	FStaticMeshConstAttributes Attributes(InMeshDescription);
+	TVertexInstanceAttributesConstRef<FVector2f> TexCoords = Attributes.GetVertexInstanceUVs();
+	TVertexInstanceAttributesConstRef<FVector4f> VertexColors = Attributes.GetVertexInstanceColors();
 
-	TVertexAttributesConstRef<FVector3f> Positions;
-	TVertexInstanceAttributesConstRef<FVector3f> Normals;
-	TVertexInstanceAttributesConstRef<FVector2f> TexCoords;
+	int32 NumVertexInstances = InMeshDescription.VertexInstances().Num();
+	int32 NumTriangles = InMeshDescription.Triangles().Num();
 
-	TArray<FTriangleID>					RemapTriangles;
-	TArray<FVertexInstanceID>			UniqueVerts;
-	TArray<FVector2D>&					OutputTexCoords;
-	bool								bMustRemap;
+	OutRemappedTriangles.Reserve(NumTriangles);
 
-	FUniqueUVMeshDescriptionView(const FMeshDescription& InMeshDescription, bool bMergeIdenticalMaterials, TArray<FVector2D>& InOutTexCoords)
-		: MeshDescription(InMeshDescription)
-		, OutputTexCoords(InOutTexCoords)
+	TMap<uint32, FTriangleID> UniqueTriangles;
+	if (bMergeIdenticalMaterials)
 	{
-		FStaticMeshConstAttributes Attributes(MeshDescription);
-		Positions = Attributes.GetVertexPositions();
-		Normals = Attributes.GetVertexInstanceNormals();
-		TexCoords = Attributes.GetVertexInstanceUVs();
-
-		TVertexInstanceAttributesConstRef<FVector4f> VertexColors = Attributes.GetVertexInstanceColors();
+		UniqueTriangles.Reserve(NumTriangles);
+	}
 
-		int32 NumVertexInstances = MeshDescription.VertexInstances().Num();
-		int32 NumTriangles = MeshDescription.Triangles().Num();
+	if (OutUniqueVerts)
+	{
+		OutUniqueVerts->Reserve(NumVertexInstances);
+	}
 
-		OutputTexCoords.SetNumZeroed(NumVertexInstances);
-		UniqueVerts.Reserve(NumVertexInstances);
-		RemapTriangles.Reserve(NumTriangles);
+	if (OutDuplicateTriangles)
+	{
+		OutDuplicateTriangles->Reserve(NumTriangles);
+	}
 
-		TMap<uint32, FTriangleID> UniqueTriangles;
-		if (bMergeIdenticalMaterials)
+	// Compute an hash value per triangle, based on its UVs & vertices colors
+	auto HashAttribute = [](FVertexInstanceID InVertexInstanceID, auto InAttributeArrayRef, int32& TriangleHash)
+	{
+		for (int32 Channel = 0; Channel < InAttributeArrayRef.GetNumChannels(); ++Channel)
 		{
-			UniqueTriangles.Reserve(NumTriangles);
+			for (const auto& Element : InAttributeArrayRef.GetArrayView(InVertexInstanceID, Channel))
+			{
+				TriangleHash = HashCombine(TriangleHash, GetTypeHash(Element));
+			}
 		}
+	};
+
+	for (const FTriangleID TriangleID : InMeshDescription.Triangles().GetElementIDs())
+	{
+		const FPolygonGroupID RefPolygonGroupID = InMeshDescription.GetTrianglePolygonGroup(TriangleID);
+		TConstArrayView<const FVertexInstanceID> VertexInstancesIDs = InMeshDescription.GetTriangleVertexInstances(TriangleID);
+		TConstArrayView<const FVertexID> VertexIDs = InMeshDescription.GetTriangleVertices(TriangleID);
 
-		auto HashAttribute = [](FVertexInstanceID InVertexInstanceID, auto InAttributeArrayRef, int32& TriangleHash)
+		FTriangleID RemapTriangleID = TriangleID;
+
+		bool bUnique = true;
+
+		if (bMergeIdenticalMaterials)
 		{
-			for (int32 Channel = 0; Channel < InAttributeArrayRef.GetNumChannels(); ++Channel)
+			int32 TriangleHash = GetTypeHash(RefPolygonGroupID);
+			for (const FVertexInstanceID& VertexInstanceID : VertexInstancesIDs)
+			{
+				// Compute hash based on UVs & vertices colors
+				HashAttribute(VertexInstanceID, TexCoords, TriangleHash);
+				HashAttribute(VertexInstanceID, VertexColors, TriangleHash);
+			}
+
+			FTriangleID* UniqueTriangleIDPtr = UniqueTriangles.Find(TriangleHash);
+			if (UniqueTriangleIDPtr != nullptr)
 			{
-				for (const auto& Element : InAttributeArrayRef.GetArrayView(InVertexInstanceID, Channel))
+				RemapTriangleID = *UniqueTriangleIDPtr;
+				bUnique = false;
+
+				if (OutDuplicateTriangles)
 				{
-					TriangleHash = HashCombine(TriangleHash, GetTypeHash(Element));
+					OutDuplicateTriangles->Add(TriangleID);
 				}
 			}
-		};
+			else
+			{
+				UniqueTriangles.Add(TriangleHash, TriangleID);
+			}
+		}
 
-		for (const FTriangleID TriangleID : MeshDescription.Triangles().GetElementIDs())
+		if (bUnique && OutUniqueVerts)
 		{
-			const FPolygonGroupID RefPolygonGroupID = MeshDescription.GetTrianglePolygonGroup(TriangleID);
-			TConstArrayView<const FVertexInstanceID> VertexInstancesIDs = MeshDescription.GetTriangleVertexInstances(TriangleID);
-			TConstArrayView<const FVertexID> VertexIDs = MeshDescription.GetTriangleVertices(TriangleID);
+			OutUniqueVerts->Append(VertexInstancesIDs);
+		}
 
-			FTriangleID RemapTriangleID = TriangleID;
+		OutRemappedTriangles.Add(RemapTriangleID);
+	}
 
-			bool bUnique = true;
+	const bool bPerformedRemapping = bMergeIdenticalMaterials && UniqueTriangles.Num() != OutRemappedTriangles.Num();
+	return bPerformedRemapping;
+}
 
-			if (bMergeIdenticalMaterials)
-			{
-				int32 TriangleHash = GetTypeHash(RefPolygonGroupID);
-				for (const FVertexInstanceID& VertexInstanceID : VertexInstancesIDs)
-				{
-					// Compute hash based on UVs & vertices colors
-					HashAttribute(VertexInstanceID, TexCoords, TriangleHash);
-					HashAttribute(VertexInstanceID, VertexColors, TriangleHash);
-				}
+template <typename TSrcUVs, typename TDstUVs>
+static void CopyRemappedUVs(const FMeshDescription& InMeshDescription, const TArray<FTriangleID>& RemappedTriangles, const FElementIDRemappings* SrcElementIDRemappings, bool bCopyOnlyRemappedTrianglesUVs, const TSrcUVs& SrcUVs, TDstUVs& DstUVs)
+{
+	int32 RemappedTrianglesIdx = 0;
+	for (const FTriangleID TriangleID : InMeshDescription.Triangles().GetElementIDs())
+	{
+		FTriangleID RemappedTriangleID = RemappedTriangles[RemappedTrianglesIdx];
 
-				FTriangleID* UniqueTriangleIDPtr = UniqueTriangles.Find(TriangleHash);
-				if (UniqueTriangleIDPtr != nullptr)
-				{
-					RemapTriangleID = *UniqueTriangleIDPtr;
-					bUnique = false;
-				}
-				else
-				{
-					UniqueTriangles.Add(TriangleHash, TriangleID);
-				}
+		if (!bCopyOnlyRemappedTrianglesUVs || RemappedTriangleID != TriangleID)
+		{
+			if (SrcElementIDRemappings)
+			{
+				RemappedTriangleID = SrcElementIDRemappings->GetRemappedTriangleID(RemappedTriangleID);
 			}
 
-			if (bUnique)
+			TConstArrayView<const FVertexInstanceID> SrcVertexInstancesIDs = InMeshDescription.GetTriangleVertexInstances(RemappedTriangleID);
+			TConstArrayView<const FVertexInstanceID> DstVertexInstancesIDs = InMeshDescription.GetTriangleVertexInstances(TriangleID);
+
+			for (int32 i = 0; i < 3; i++)
 			{
-				UniqueVerts.Append(VertexInstancesIDs);
+				DstUVs[DstVertexInstancesIDs[i]] = FVector2D(SrcUVs[SrcVertexInstancesIDs[i]]);
 			}
-
-			RemapTriangles.Add(RemapTriangleID);
 		}
 
-		bMustRemap = bMergeIdenticalMaterials && UniqueTriangles.Num() != RemapTriangles.Num();
+		RemappedTrianglesIdx++;
+	}
+}
+
+// Mesh view that will expose only unique UVs if bMergeIdenticalMaterials is provided
+struct FUniqueUVMeshDescriptionView final : FLayoutUV::IMeshView
+{
+	const FMeshDescription& MeshDescription;
+
+	TVertexAttributesConstRef<FVector3f> Positions;
+	TVertexInstanceAttributesConstRef<FVector3f> Normals;
+	TVertexInstanceAttributesConstRef<FVector2f> TexCoords;
+
+	TArray<FTriangleID>					RemapTriangles;
+	TArray<FVertexInstanceID>			UniqueVerts;
+	TArray<FVector2D>&					OutputTexCoords;
+	bool								bMustRemap;
+
+	FUniqueUVMeshDescriptionView(const FMeshDescription& InMeshDescription, bool bMergeIdenticalMaterials, TArray<FVector2D>& InOutTexCoords)
+		: MeshDescription(InMeshDescription)
+		, OutputTexCoords(InOutTexCoords)
+	{
+		FStaticMeshConstAttributes Attributes(MeshDescription);
+		Positions = Attributes.GetVertexPositions();
+		Normals = Attributes.GetVertexInstanceNormals();
+		TexCoords = Attributes.GetVertexInstanceUVs();
+
+		OutputTexCoords.SetNumZeroed(MeshDescription.VertexInstances().Num());
+
+		bMustRemap = GatherUniqueTriangles(MeshDescription, bMergeIdenticalMaterials, RemapTriangles, &UniqueVerts, nullptr);
 	}
 
 	uint32 GetNumIndices() const override
@@ -2122,28 +2174,22 @@ struct FUniqueUVMeshDescriptionView final : FLayoutUV::IMeshView
 	{
 		if (bMustRemap)
 		{
-		    int32 RemapTrianglesIdx = 0;
-		    for (const FTriangleID TriangleID : MeshDescription.Triangles().GetElementIDs())
-		    {
-			    FTriangleID RemapTriangleID = RemapTriangles[RemapTrianglesIdx];
-			    if (RemapTriangleID != TriangleID)
-			    {
-				    TConstArrayView<const FVertexInstanceID> SrcVertexInstancesIDs = MeshDescription.GetTriangleVertexInstances(RemapTriangleID);
-				    TConstArrayView<const FVertexInstanceID> DstVertexInstancesIDs = MeshDescription.GetTriangleVertexInstances(TriangleID);
-    
-				    for (int32 i = 0; i < 3; i++)
-				    {
-					    OutputTexCoords[DstVertexInstancesIDs[i]] = OutputTexCoords[SrcVertexInstancesIDs[i]];
-				    }
-			    }
-    
-			    RemapTrianglesIdx++;
-		    }
+			CopyRemappedUVs(MeshDescription, RemapTriangles, nullptr, true, OutputTexCoords, OutputTexCoords);
 		}
 	}
 };
 
 bool FStaticMeshOperations::GenerateUniqueUVsForStaticMesh(const FMeshDescription& MeshDescription, int32 TextureResolution, bool bMergeIdenticalMaterials, TArray<FVector2D>& OutTexCoords)
+{
+	FGenerateUVOptions GenerateUVOptions;
+	GenerateUVOptions.TextureResolution = TextureResolution;
+	GenerateUVOptions.bMergeTrianglesWithIdenticalAttributes = bMergeIdenticalMaterials;
+	GenerateUVOptions.UVMethod = EGenerateUVMethod::Legacy;
+
+	return GenerateUV(MeshDescription, GenerateUVOptions, OutTexCoords);
+}
+
+bool FStaticMeshOperations::GenerateUV(const FMeshDescription& MeshDescription, const FGenerateUVOptions& GenerateUVOptions, TArray<FVector2D>& OutTexCoords)
 {
 	TRACE_CPUPROFILER_EVENT_SCOPE(FStaticMeshOperations::GenerateUniqueUVsForStaticMesh)
 
@@ -2151,10 +2197,12 @@ bool FStaticMeshOperations::GenerateUniqueUVsForStaticMesh(const FMeshDescriptio
 
 	OutTexCoords.Reset();
 
+	const bool bAutoUVAvailable = WITH_EDITOR;
 	const bool bHasUVs = VertexInstanceUVs.GetNumElements() > 0;
-	if (bHasUVs)
+	const bool bUseLegacy = GenerateUVOptions.UVMethod == EGenerateUVMethod::Legacy || !bAutoUVAvailable;
+	if (bHasUVs && bUseLegacy)
 	{
-		FUniqueUVMeshDescriptionView MeshDescriptionView(MeshDescription, bMergeIdenticalMaterials, OutTexCoords);
+		FUniqueUVMeshDescriptionView MeshDescriptionView(MeshDescription, GenerateUVOptions.bMergeTrianglesWithIdenticalAttributes, OutTexCoords);
 
 		// Find overlapping corners for UV generator. Allow some threshold - this should not produce any error in a case if resulting
 		// mesh will not merge these vertices.
@@ -2166,7 +2214,7 @@ bool FStaticMeshOperations::GenerateUniqueUVsForStaticMesh(const FMeshDescriptio
 
 		// Scale down texture resolution to speed up UV generation time
 		// Packing expects at least one texel per chart. This is the absolute minimum to generate valid UVs.
-		const int32 PackingResolution = FMath::Clamp(TextureResolution / 4, 32, 512);
+		const int32 PackingResolution = FMath::Clamp(GenerateUVOptions.TextureResolution / 4, 32, 512);
 		const int32 AbsoluteMinResolution = 1 << FMath::CeilLogTwo(FMath::Sqrt((float)NumCharts));
 		const int32 FinalPackingResolution = FMath::Max(PackingResolution, AbsoluteMinResolution);
 
@@ -2189,11 +2237,43 @@ bool FStaticMeshOperations::GenerateUniqueUVsForStaticMesh(const FMeshDescriptio
 		IGeometryProcessingInterfacesModule* GeomProcInterfaces = FModuleManager::Get().GetModulePtr<IGeometryProcessingInterfacesModule>("GeometryProcessingInterfaces");
 		if (GeomProcInterfaces)
 		{
-		    IGeometryProcessing_MeshAutoUV* MeshAutoUV = GeomProcInterfaces->GetMeshAutoUVImplementation();
-    
 		    FMeshDescription MeshCopy = MeshDescription;
-    
+			TArray<FTriangleID>	RemapTriangles;
+			FElementIDRemappings ElementIDRemappings;
+
+			bool bMustRemap = false;
+
+			// Ensure we have properly setup TriangleUVs on our mesh
+			int32 UVChannelCount = MeshCopy.VertexInstanceAttributes().GetAttributeChannelCount(MeshAttribute::VertexInstance::TextureCoordinate);
+			MeshCopy.SetNumUVChannels(UVChannelCount);
+
+			if (GenerateUVOptions.bMergeTrianglesWithIdenticalAttributes)
+			{
+				TArray<FTriangleID> DuplicateTriangles;
+
+				bMustRemap = GatherUniqueTriangles(MeshDescription, true, RemapTriangles, nullptr, &DuplicateTriangles);
+				if (bMustRemap)
+				{
+					MeshCopy.DeleteTriangles(DuplicateTriangles);
+					MeshCopy.Compact(ElementIDRemappings);
+				}
+			}
+
+			auto GetAutoUVMethod = [](EGenerateUVMethod GenerateUVMethod) -> IGeometryProcessing_MeshAutoUV::EAutoUVMethod
+			{
+				switch (GenerateUVMethod)
+				{
+				case EGenerateUVMethod::UVAtlas:	return IGeometryProcessing_MeshAutoUV::EAutoUVMethod::UVAtlas;
+				case EGenerateUVMethod::XAtlas:		return IGeometryProcessing_MeshAutoUV::EAutoUVMethod::XAtlas;
+				default:							return IGeometryProcessing_MeshAutoUV::EAutoUVMethod::PatchBuilder;
+				}
+			};
+			    
+			IGeometryProcessing_MeshAutoUV* MeshAutoUV = GeomProcInterfaces->GetMeshAutoUVImplementation();
+
 		    IGeometryProcessing_MeshAutoUV::FOptions Options = MeshAutoUV->ConstructDefaultOptions();
+			Options.Method = GetAutoUVMethod(GenerateUVOptions.UVMethod);
+
 		    IGeometryProcessing_MeshAutoUV::FResults Results;
 		    MeshAutoUV->GenerateUVs(MeshCopy, Options, Results);
     
@@ -2202,11 +2282,22 @@ bool FStaticMeshOperations::GenerateUniqueUVsForStaticMesh(const FMeshDescriptio
 			    TVertexInstanceAttributesConstRef<FVector2f> TexCoords;
 			    FStaticMeshConstAttributes AttributesCopy(MeshCopy);
 			    TexCoords = AttributesCopy.GetVertexInstanceUVs();
+
+				OutTexCoords.SetNumUninitialized(MeshDescription.VertexInstances().Num());
     
-			    for (const FVertexInstanceID VertexInstanceID : MeshCopy.VertexInstances().GetElementIDs())
-			    {
-				    OutTexCoords.Add(FVector2D(TexCoords.Get(VertexInstanceID, 0)));
-			    }
+				if (bMustRemap)
+				{
+					CopyRemappedUVs(MeshDescription, RemapTriangles, &ElementIDRemappings, false, TexCoords, OutTexCoords);
+				}
+				else
+				{
+					int32 VertexInstanceIndex = 0;
+					for (const FVertexInstanceID VertexInstanceID : MeshCopy.VertexInstances().GetElementIDs())
+					{
+						OutTexCoords[VertexInstanceIndex] = FVector2D(TexCoords.Get(VertexInstanceID, 0));
+						VertexInstanceIndex++;
+					}
+				}
 		    }
 		}
 	}
diff --git a/Engine/Source/Runtime/StaticMeshDescription/Public/StaticMeshOperations.h b/Engine/Source/Runtime/StaticMeshDescription/Public/StaticMeshOperations.h
--- a/Engine/Source/Runtime/StaticMeshDescription/Public/StaticMeshOperations.h
+++ b/Engine/Source/Runtime/StaticMeshDescription/Public/StaticMeshOperations.h
@@ -130,7 +130,32 @@ public:
 		const FOverlappingCorners& OverlappingCorners);
 
 	/** Create some UVs from the specified mesh description data. */
+	UE_DEPRECATED(5.4, "Please use GenerateUV() instead.")
 	static STATICMESHDESCRIPTION_API bool GenerateUniqueUVsForStaticMesh(const FMeshDescription& MeshDescription, int32 TextureResolution, bool bMergeIdenticalMaterials, TArray<FVector2D>& OutTexCoords);
+		
+	enum class EGenerateUVMethod
+	{
+		Default,
+		Legacy,
+		UVAtlas,
+		XAtlas,		
+		PatchBuilder
+	};
+
+	struct FGenerateUVOptions
+	{
+		// Expected texture resolution
+		int32 TextureResolution = 512;
+
+		// Wether to fold triangles sharing the same UVs & vertex colors in the generated UV mapping.
+		bool bMergeTrianglesWithIdenticalAttributes = false;
+
+		// Method to use when generating UVs
+		EGenerateUVMethod UVMethod = EGenerateUVMethod::Default;
+	};
+
+	/** Generate UV coordinates from the specified mesh description data. */
+	static STATICMESHDESCRIPTION_API bool GenerateUV(const FMeshDescription& MeshDescription, const FGenerateUVOptions& Options, TArray<FVector2D>& OutTexCoords);
 
 	/** Add a UV channel to the MeshDescription. */
 	static STATICMESHDESCRIPTION_API bool AddUVChannel(FMeshDescription& MeshDescription);
