add option to StaticMeshOperations::TransformMesh to correctly transform the normals and tangents. Default to the current behavior, which instead matches UE renderer. Add comments explaining both paths.

diff --git a/Engine/Source/Runtime/StaticMeshDescription/Private/StaticMeshOperations.cpp b/Engine/Source/Runtime/StaticMeshDescription/Private/StaticMeshOperations.cpp
--- a/Engine/Source/Runtime/StaticMeshDescription/Private/StaticMeshOperations.cpp
+++ b/Engine/Source/Runtime/StaticMeshDescription/Private/StaticMeshOperations.cpp
@@ -2722,12 +2722,12 @@ void FStaticMeshOperations::FlipPolygons(FMeshDescription& MeshDescription)
 	}
 }
 
-void FStaticMeshOperations::ApplyTransform(FMeshDescription& MeshDescription, const FTransform& Transform)
+void FStaticMeshOperations::ApplyTransform(FMeshDescription& MeshDescription, const FTransform& Transform, bool bApplyCorrectNormalTransform)
 {
-	ApplyTransform(MeshDescription, Transform.ToMatrixWithScale());
+	ApplyTransform(MeshDescription, Transform.ToMatrixWithScale(), bApplyCorrectNormalTransform);
 }
 
-void FStaticMeshOperations::ApplyTransform(FMeshDescription& MeshDescription, const FMatrix& Transform)
+void FStaticMeshOperations::ApplyTransform(FMeshDescription& MeshDescription, const FMatrix& Transform, bool bApplyCorrectNormalTransform)
 {
 	TRACE_CPUPROFILER_EVENT_SCOPE(FStaticMeshOperations::ApplyTransform)
 
@@ -2741,19 +2741,33 @@ void FStaticMeshOperations::ApplyTransform(FMeshDescription& MeshDescription, co
 		VertexPositions[VertexID] = FVector4f(Transform.TransformPosition(FVector3d(VertexPositions[VertexID])));
 	}
 
-	FMatrix TransformInverseTransposeMatrix = Transform.Inverse().GetTransposed();
-	TransformInverseTransposeMatrix.RemoveScaling();
-
 	const bool bIsMirrored = Transform.Determinant() < 0.f;
 	const float MulBy = bIsMirrored ? -1.f : 1.f;
 
+	FMatrix NormalsTransform, TangentsTransform;
+	if (bApplyCorrectNormalTransform)
+	{
+		// Note: Assuming we'll normalize after, transforming by the transpose-adjoint * the sign of the determinant
+		// is equivalent to transforming by the inverse transpose; ref: TMatrix::TransformByUsingAdjointT
+		NormalsTransform = Transform.TransposeAdjoint() * (double)MulBy;
+		// Note: Tangents *do not* transform by the Transform's inverse transpose, just by the Transform
+		TangentsTransform = Transform;
+	}
+	else // match UE renderer
+	{
+		// UE's renderer transforms normals and tangents without scale (as in FTransform::TransformVectorNoScale)
+		NormalsTransform = Transform;
+		NormalsTransform.RemoveScaling();
+		TangentsTransform = NormalsTransform;
+	}
+
 	for (const FVertexInstanceID VertexInstanceID : MeshDescription.VertexInstances().GetElementIDs())
 	{
 		FVector3f Tangent = VertexInstanceTangents[VertexInstanceID];
 		FVector3f Normal = VertexInstanceNormals[VertexInstanceID];
 
-		VertexInstanceTangents[VertexInstanceID] = (FVector3f)FVector(TransformInverseTransposeMatrix.TransformVector((FVector)Tangent).GetSafeNormal());
-		VertexInstanceNormals[VertexInstanceID] = (FVector3f)FVector(TransformInverseTransposeMatrix.TransformVector((FVector)Normal).GetSafeNormal());
+		VertexInstanceTangents[VertexInstanceID] = (FVector3f)FVector(TangentsTransform.TransformVector((FVector)Tangent).GetSafeNormal());
+		VertexInstanceNormals[VertexInstanceID] = (FVector3f)FVector(NormalsTransform.TransformVector((FVector)Normal).GetSafeNormal());
 
 		float BinormalSign = VertexInstanceBinormalSigns[VertexInstanceID];
 		VertexInstanceBinormalSigns[VertexInstanceID] = BinormalSign * MulBy;
diff --git a/Engine/Source/Runtime/StaticMeshDescription/Public/StaticMeshOperations.h b/Engine/Source/Runtime/StaticMeshDescription/Public/StaticMeshOperations.h
--- a/Engine/Source/Runtime/StaticMeshDescription/Public/StaticMeshOperations.h
+++ b/Engine/Source/Runtime/StaticMeshDescription/Public/StaticMeshOperations.h
@@ -205,9 +205,12 @@ public:
 	/** Flip the facing for a set of input polygons. */
 	static STATICMESHDESCRIPTION_API void FlipPolygons(FMeshDescription& MeshDescription);
 
-	/** Transforms the MeshDescription data using the provided transform. */
-	static STATICMESHDESCRIPTION_API void ApplyTransform(FMeshDescription& MeshDescription, const FTransform& Transform);
-	static STATICMESHDESCRIPTION_API void ApplyTransform(FMeshDescription& MeshDescription, const FMatrix& Transform);	
+	/** 
+	 * Transforms the MeshDescription data using the provided transform.
+	 * @param bApplyCorrectNormalTransform Whether to correctly transform normals and tangents. Otherwise, will match the UE renderer and transform them without scale.
+	 */
+	static STATICMESHDESCRIPTION_API void ApplyTransform(FMeshDescription& MeshDescription, const FTransform& Transform, bool bApplyCorrectNormalTransform = false);
+	static STATICMESHDESCRIPTION_API void ApplyTransform(FMeshDescription& MeshDescription, const FMatrix& Transform, bool bApplyCorrectNormalTransform = false);
 
 	/**
 	 * Return the number of unique vertices, unique vertices are the result of welding all similar vertex instances (position, UV, tangent space, color,...)
