Skeletal Mesh: When converting from import model to mesh description, don't ditch all normals/tangents even if some are broken. Instead, re-create the broken normals/tangents as needed.

#rnx

diff --git a/Engine/Source/Runtime/Engine/Private/Rendering/SkeletalMeshLODImporterData.cpp b/Engine/Source/Runtime/Engine/Private/Rendering/SkeletalMeshLODImporterData.cpp
--- a/Engine/Source/Runtime/Engine/Private/Rendering/SkeletalMeshLODImporterData.cpp
+++ b/Engine/Source/Runtime/Engine/Private/Rendering/SkeletalMeshLODImporterData.cpp
@@ -2,8 +2,6 @@
 
 #include "Rendering/SkeletalMeshLODImporterData.h"
 
-#include "UObject/Package.h"
-
 #if WITH_EDITOR
 
 #include "Algo/AnyOf.h"
@@ -1904,7 +1902,7 @@ static void CopySkinWeightsToAttribute(
 	}
 }
 
-bool FSkeletalMeshImportData::GetMeshDescription(const USkeletalMesh* InSkeletalMesh, FMeshDescription& OutMeshDescription) const
+bool FSkeletalMeshImportData::GetMeshDescription(const USkeletalMesh* InSkeletalMesh, const FSkeletalMeshBuildSettings* InBuildSettings, FMeshDescription& OutMeshDescription) const
 {
 	using namespace UE::AnimationCore;
 	
@@ -2028,7 +2026,7 @@ bool FSkeletalMeshImportData::GetMeshDescription(const USkeletalMesh* InSkeletal
 				continue;
 			}
 
-			// The set of bound bones must completely include the bones in the base mesh. Otherwise we cannot map the influences over to
+			// The set of bound bones must completely include the bones in the base mesh. Otherwise, we cannot map the influences over to
 			// the base skeleton.
 			TSet<FString> BaseBones;
 			TSet<int32> MissingBoneIndices;
@@ -2202,33 +2200,6 @@ bool FSkeletalMeshImportData::GetMeshDescription(const USkeletalMesh* InSkeletal
 		MaterialGroups.Add(PolygonGroupID);
 	}
 
-	// These are not always honest.
-	bool bHaveValidNormals = bHasNormals;
-	bool bHaveValidTangents = bHasNormals;
-	
-	for (int32 TriangleIndex = 0; TriangleIndex < Faces.Num(); TriangleIndex++)
-	{
-		const SkeletalMeshImportData::FTriangle &Triangle = Faces[TriangleIndex];
-
-		for (int32 Corner = 0; Corner < 3; Corner++)
-		{
-			if (!Triangle.TangentZ[Corner].IsNormalized())
-			{
-				bHaveValidNormals = false;
-			}
-			if (!Triangle.TangentX[Corner].IsNormalized() ||
-			    !Triangle.TangentY[Corner].IsNormalized())
-			{
-				bHaveValidTangents = false;
-			}
-
-			if (!bHaveValidNormals && !bHaveValidTangents)
-			{
-				break;
-			}
-		}
-	}
-
 	for (int32 TriangleIndex = 0; TriangleIndex < Faces.Num(); TriangleIndex++)
 	{
 		const SkeletalMeshImportData::FTriangle &Triangle = Faces[TriangleIndex];
@@ -2254,16 +2225,16 @@ bool FSkeletalMeshImportData::GetMeshDescription(const USkeletalMesh* InSkeletal
 					VertexInstanceUVs.Set(VertexInstanceID, UVIndex, Wedge.UVs[UVIndex]);
 				}
 
-				if (bHaveValidNormals)
+				if (bHasNormals)
 				{
 					VertexInstanceNormals.Set(VertexInstanceID, Triangle.TangentZ[Corner]);
 				}
-				if (bHaveValidTangents)
+				if (bHasTangents)
 				{
 					VertexInstanceTangents.Set(VertexInstanceID, Triangle.TangentX[Corner]);
 
 					// We can only divine the bi-tangent sign if the normal is also given. 
-					if (bHaveValidNormals)
+					if (bHasNormals)
 					{
 						VertexInstanceBinormalSigns.Set(VertexInstanceID,
 							((Triangle.TangentZ[Corner] ^ Triangle.TangentX[Corner]) | Triangle.TangentY[Corner]) < 0 ? -1.0f : 1.0f);
@@ -2393,8 +2364,41 @@ bool FSkeletalMeshImportData::GetMeshDescription(const USkeletalMesh* InSkeletal
 		}
 	}
 
+	// We don't need to do this when retrieving alt skin weights.
+	FString SkeletalMeshPath;
+	if (InSkeletalMesh)
+	{
+		SkeletalMeshPath = InSkeletalMesh->GetOuter()->GetPathName();		
+	}
+	
 	FSkeletalMeshOperations::ConvertSmoothGroupToHardEdges(FaceSmoothingMasks, OutMeshDescription);
 
+	// Check if we have any broken data, including UVs and normals/tangents.
+	FSkeletalMeshOperations::ValidateAndFixData(OutMeshDescription, SkeletalMeshPath);
+
+	bool bNormalsValid, bTangentsValid;
+	FStaticMeshOperations::AreNormalsAndTangentsValid(OutMeshDescription, bNormalsValid, bTangentsValid);
+	if (!bNormalsValid || !bTangentsValid)
+	{
+		// This is required by FSkeletalMeshOperations::ComputeTangentsAndNormals to function correctly.
+		FSkeletalMeshOperations::ComputeTriangleTangentsAndNormals(OutMeshDescription, UE_SMALL_NUMBER, !SkeletalMeshPath.IsEmpty() ? *SkeletalMeshPath : nullptr);
+
+		EComputeNTBsFlags ComputeNTBsOptions = EComputeNTBsFlags::None;
+		if (InBuildSettings && InBuildSettings->bComputeWeightedNormals)
+		{
+			ComputeNTBsOptions |= EComputeNTBsFlags::WeightedNTBs;
+		}
+
+		// This only recomputes broken normals/tangents. The ValidateAndFixData function above will have turned all non-finite normals and tangents into
+		// zero vectors.
+		FSkeletalMeshOperations::ComputeTangentsAndNormals(OutMeshDescription, ComputeNTBsOptions);
+
+		// We don't need the triangle tangents and normals anymore.
+		OutMeshDescription.TriangleAttributes().UnregisterAttribute(MeshAttribute::Triangle::Normal);
+		OutMeshDescription.TriangleAttributes().UnregisterAttribute(MeshAttribute::Triangle::Tangent);
+		OutMeshDescription.TriangleAttributes().UnregisterAttribute(MeshAttribute::Triangle::Binormal);
+	}
+
 	// Convert the MeshInfo data.
 	if (!MeshInfos.IsEmpty())
 	{
@@ -2495,7 +2499,7 @@ void FSkeletalMeshImportData::CopySkinWeightsToMeshDescription(
 
 	// The topologies don't match, proceed as above.
 	FMeshDescription AlternateInfluenceMesh;
-	InSkinWeightMesh.GetMeshDescription(nullptr, AlternateInfluenceMesh);
+	InSkinWeightMesh.GetMeshDescription(nullptr, nullptr, AlternateInfluenceMesh);
 	
 	FSkeletalMeshOperations::CopySkinWeightAttributeFromMesh(
 	AlternateInfluenceMesh, OutMeshDescription, NAME_None, InSkinWeightName, &AltMeshBoneToBaseBoneMap);   
@@ -2551,7 +2555,11 @@ static void ConvertHardEdgesToSmoothMasks(
 		CurrentSmoothMask <<= 1;
 		if (CurrentSmoothMask == 0)
 		{
-			CurrentSmoothMask = 1;
+			// If we exhausted all available bits, then thunk to the more complete algorithm. For reasons unknown at this time, it doesn't generate
+			// nice smooth groups for some simpler test objects. For more complex input products it does a decent job though.
+			OutSmoothMasks.SetNumZeroed(InMeshDescription.Triangles().Num());
+			FStaticMeshOperations::ConvertHardEdgesToSmoothGroup(InMeshDescription, OutSmoothMasks);
+			break;
 		}
 	}
 }
@@ -2727,20 +2735,23 @@ FSkeletalMeshImportData FSkeletalMeshImportData::CreateFromMeshDescription(const
 		}
 	}
 
-	
-	bool bHaveValidNormals = true;
-	bool bHaveValidTangents = true;
+	bool bHaveValidNormals = false;
+	bool bHaveValidTangents = false;
 	for (FVertexInstanceID VertexInstanceID: InMeshDescription.VertexInstances().GetElementIDs())
 	{
-		if (!VertexInstanceNormals.Get(VertexInstanceID).IsNormalized())
+		const FVector3f Normal = VertexInstanceNormals.Get(VertexInstanceID);
+		if (!Normal.IsNearlyZero(UE_SMALL_NUMBER) && !Normal.ContainsNaN())
 		{
-			bHaveValidNormals = false;
+			bHaveValidNormals = true;
 		}
-		if (!VertexInstanceTangents.Get(VertexInstanceID).IsNormalized())
+		const FVector3f Tangent = VertexInstanceTangents.Get(VertexInstanceID);
+		if (!Tangent.IsNearlyZero(UE_SMALL_NUMBER) && !Tangent.ContainsNaN())
 		{
-			bHaveValidTangents = false;
+			bHaveValidTangents = true;
 		}
-		if (!bHaveValidNormals && !bHaveValidTangents)
+
+		// If we found any valid normals/tangents, we can stop now.
+		if (bHaveValidNormals && bHaveValidTangents)
 		{
 			break;
 		}
@@ -2797,6 +2808,8 @@ FSkeletalMeshImportData FSkeletalMeshImportData::CreateFromMeshDescription(const
 			}
 			else
 			{
+				// The normal/tangent computation during conversion to render data will automatically regenerate any degenerate normals.
+				// Same for the tangent/binormal below.
 				Face.TangentZ[Corner] = FVector3f::ZeroVector;
 			}
 
@@ -2808,6 +2821,7 @@ FSkeletalMeshImportData FSkeletalMeshImportData::CreateFromMeshDescription(const
 			else
 			{
 				Face.TangentX[Corner] = FVector3f::ZeroVector;
+				Face.TangentY[Corner] = FVector3f::ZeroVector;
 			}
 
 			const int32 WedgeIndex = VertexInstanceID.GetValue();
diff --git a/Engine/Source/Runtime/Engine/Private/SkeletalMesh.cpp b/Engine/Source/Runtime/Engine/Private/SkeletalMesh.cpp
--- a/Engine/Source/Runtime/Engine/Private/SkeletalMesh.cpp
+++ b/Engine/Source/Runtime/Engine/Private/SkeletalMesh.cpp
@@ -2585,7 +2585,7 @@ void USkeletalMesh::SaveLODImportedData(const int32 LODIndex, const FSkeletalMes
 	
 	FMeshDescription MeshDescription;
 
-	if (InMesh.GetMeshDescription(nullptr, MeshDescription))
+	if (InMesh.GetMeshDescription(nullptr, &GetLODInfo(LODIndex)->BuildSettings, MeshDescription))
 	{
 		CreateMeshDescription(LODIndex, MoveTemp(MeshDescription));
 		CommitMeshDescription(LODIndex);
@@ -3277,7 +3277,7 @@ void USkeletalMesh::BeginPostLoadInternal(FSkinnedAssetPostLoadContext& Context)
 						}
 
 						FMeshDescription MeshDescription;
-						if (SkeletalMeshImportData.GetMeshDescription(this, MeshDescription))
+						if (SkeletalMeshImportData.GetMeshDescription(this, &GetLODInfo(LODIndex)->BuildSettings, MeshDescription))
 						{
 							CreateMeshDescription(LODIndex, MoveTemp(MeshDescription));
 							CommitMeshDescription(LODIndex);
diff --git a/Engine/Source/Runtime/Engine/Private/SkeletalMeshSourceModel.cpp b/Engine/Source/Runtime/Engine/Private/SkeletalMeshSourceModel.cpp
--- a/Engine/Source/Runtime/Engine/Private/SkeletalMeshSourceModel.cpp
+++ b/Engine/Source/Runtime/Engine/Private/SkeletalMeshSourceModel.cpp
@@ -6,6 +6,7 @@
 
 #if WITH_EDITOR
 #include "Engine/SkeletalMesh.h"
+#include "Engine/SkinnedAssetCommon.h"
 #include "SkeletalMeshDescription.h"
 #include "Rendering/SkeletalMeshLODImporterData.h"
 #endif
@@ -329,7 +330,7 @@ void FSkeletalMeshSourceModel::ConvertRawMeshToMeshDescriptionBulkData()
 	}
 	
 	FMeshDescription& MeshDescription = MeshDescriptionBulkData->CreateMeshDescription()->GetMeshDescription(); 
-	ImportData.GetMeshDescription(GetOwner(), MeshDescription);
+	ImportData.GetMeshDescription(SkeletalMesh, &SkeletalMesh->GetLODInfo(RawMeshBulkDataLODIndex)->BuildSettings, MeshDescription);
 
 	UpdateCachedMeshStatistics(&MeshDescription);
 
diff --git a/Engine/Source/Runtime/Engine/Public/Rendering/SkeletalMeshLODImporterData.h b/Engine/Source/Runtime/Engine/Public/Rendering/SkeletalMeshLODImporterData.h
--- a/Engine/Source/Runtime/Engine/Public/Rendering/SkeletalMeshLODImporterData.h
+++ b/Engine/Source/Runtime/Engine/Public/Rendering/SkeletalMeshLODImporterData.h
@@ -432,9 +432,9 @@ public:
 	
 	/**
 	 * Returns a mesh description from the import data. If logging on failures is required, pass in a pointer 
-	 * to the owning skeletal mesh. Otherwise leave as a \c nullptr. 
+	 * to the owning skeletal mesh. Otherwise, leave as a \c nullptr. 
 	 */
-	ENGINE_API bool GetMeshDescription(const USkeletalMesh* InSkeletalMesh, FMeshDescription& OutMeshDescription) const;
+	ENGINE_API bool GetMeshDescription(const USkeletalMesh* InSkeletalMesh, const FSkeletalMeshBuildSettings* InBuildSettings, FMeshDescription& OutMeshDescription) const;
 
 	/**
 	 * @note MeshDescription always contains color, normal and tangent data by default. Therefore, while iterating
diff --git a/Engine/Source/Runtime/StaticMeshDescription/Private/StaticMeshOperations.cpp b/Engine/Source/Runtime/StaticMeshDescription/Private/StaticMeshOperations.cpp
--- a/Engine/Source/Runtime/StaticMeshDescription/Private/StaticMeshOperations.cpp
+++ b/Engine/Source/Runtime/StaticMeshDescription/Private/StaticMeshOperations.cpp
@@ -187,7 +187,7 @@ void FStaticMeshOperations::ComputePolygonTangentsAndNormals(FMeshDescription& M
 }
 
 
-static TTuple<FVector3f, FVector3f, FVector3f> GetTriangleTangentsAndNormals(float ComparisonThreshold, TArrayView<const FVector3f> VertexPositions, TArrayView<const FVector2D> VertexUVs)
+static TTuple<FVector3f, FVector3f, FVector3f> GetTriangleTangentsAndNormalsWithUV(float ComparisonThreshold, TArrayView<const FVector3f> VertexPositions, TArrayView<const FVector2D> VertexUVs)
 {
 	// GetSafeNormal compare the squareSum to the tolerance.
 	const float SquareComparisonThreshold = FMath::Max(ComparisonThreshold * ComparisonThreshold, MIN_flt);
@@ -249,6 +249,49 @@ static TTuple<FVector3f, FVector3f, FVector3f> GetTriangleTangentsAndNormals(flo
 	}
 }
 
+// Create a normal using the triangle plane, but use the Duff & Frisvad algorithm (see Duff 2017 in JCGT) to construct a consistent tangent from a
+// single vector. 
+static TTuple<FVector3f, FVector3f, FVector3f> GetTriangleTangentsAndNormalsWithNoUVs(float ComparisonThreshold, TArrayView<const FVector3f> VertexPositions)
+{
+	// GetSafeNormal compare the squareSum to the tolerance.
+	const float SquareComparisonThreshold = FMath::Max(ComparisonThreshold * ComparisonThreshold, MIN_flt);
+
+	const FVector3f Position0 = VertexPositions[0];
+	// If the positions deltas are too small, we get a zero vector out.
+	const FVector3f DPosition1 = VertexPositions[1] - Position0;
+	const FVector3f DPosition2 = VertexPositions[2] - Position0;
+
+	// We have a left-handed coordinate system, but a counter-clockwise winding order
+	// Hence normal calculation has to take the triangle vectors cross product in reverse.
+	// If we got a zero vector out above, then this is also zero
+	FVector3f Normal = FVector3f::CrossProduct(DPosition2, DPosition1).GetSafeNormal(SquareComparisonThreshold);
+	if (!Normal.Normalize(ComparisonThreshold))
+	{
+		return MakeTuple(FVector3f::ZeroVector, FVector3f::ZeroVector, FVector3f::ZeroVector);
+	}
+
+	FVector3f Tangent, Binormal;
+	if (Normal.Z < 0.0f)
+	{
+		const float A = 1.0f / (1.0f - Normal.Z);
+		const float B = Normal.X * Normal.Y * A;
+
+		Tangent = FVector3f(1.0f - Normal.X * Normal.X * A, -B, Normal.X);
+		Binormal = FVector3f(B, Normal.Y * Normal.Y * A - 1.0f, -Normal.Y);
+	}
+	else
+	{
+		const float A = 1.0f / (1.0f + Normal.Z);
+		const float B = -Normal.X * Normal.Y * A;
+
+		Tangent = FVector3f(1.0f - Normal.X * Normal.X * A, B, -Normal.X);
+		Binormal = FVector3f(B, 1.0f - Normal.Y * Normal.Y * A, -Normal.Y);
+	}
+
+	// The above algorithm guarantees orthogonality and normalization of the tangent & binormal if the normal vector is already normalized.
+	return MakeTuple(Normal, Tangent, Binormal);
+}
+
 
 void FStaticMeshOperations::ComputeTriangleTangentsAndNormals(FMeshDescription& MeshDescription, float ComparisonThreshold, const TCHAR* DebugName)
 {
@@ -273,7 +316,7 @@ void FStaticMeshOperations::ComputeTriangleTangentsAndNormals(FMeshDescription&
 		[BatchSize, ComparisonThreshold, NumTriangles, &Attributes, DebugName](int32 BatchIndex)
 		{
 			TArrayView<const FVector3f> VertexPositions = Attributes.GetVertexPositions().GetRawArray();
-			TArrayView<const FVector2f> VertexUVs = Attributes.GetVertexInstanceUVs().GetRawArray();
+			TArrayView<const FVector2f> VertexUVs;
 			TArrayView<const FVertexID> TriangleVertexIDs = Attributes.GetTriangleVertexIndices().GetRawArray();
 			TArrayView<const FVertexInstanceID> TriangleVertexInstanceIDs = Attributes.GetTriangleVertexInstanceIndices().GetRawArray();
 
@@ -281,6 +324,11 @@ void FStaticMeshOperations::ComputeTriangleTangentsAndNormals(FMeshDescription&
 			TArrayView<FVector3f> TriangleTangents = Attributes.GetTriangleTangents().GetRawArray();
 			TArrayView<FVector3f> TriangleBinormals = Attributes.GetTriangleBinormals().GetRawArray();
 
+			if (Attributes.GetVertexInstanceUVs().GetNumChannels() > 0)
+			{
+				VertexUVs = Attributes.GetVertexInstanceUVs().GetRawArray(0);
+			}
+
 			int32 StartIndex = BatchIndex * BatchSize;
 			int32 TriIndex = StartIndex * 3;
 			int32 EndIndex = FMath::Min(StartIndex + BatchSize, NumTriangles);
@@ -303,36 +351,44 @@ void FStaticMeshOperations::ComputeTriangleTangentsAndNormals(FMeshDescription&
 					TriangleVertexPositions[1].ContainsNaN() ||
 					TriangleVertexPositions[2].ContainsNaN())
 				{
-					UE_LOG(LogStaticMeshOperations, Warning, TEXT("Static Mesh %s has NaNs in it's vertex positions! Triangle index %d -- using identity for tangent basis."), DebugName ? DebugName : TEXT("<null>"), StartIndex);
+					UE_CLOG(DebugName != nullptr, LogStaticMeshOperations, Warning, TEXT("Static Mesh %s has NaNs in it's vertex positions! Triangle index %d -- using identity for tangent basis."), DebugName, StartIndex);
 					TriangleNormals[StartIndex] = FVector3f(1, 0, 0);
 					TriangleTangents[StartIndex] = FVector3f(0, 1, 0);
 					TriangleBinormals[StartIndex] = FVector3f(0, 0, 1);
 					continue;
 				}
 
-				FVector2D TriangleUVs[3] =
+				TTuple<FVector3f, FVector3f, FVector3f> Result;
+				if (!VertexUVs.IsEmpty())
 				{
-					FVector2D(VertexUVs[TriangleVertexInstanceIDs[TriIndex]]),
-					FVector2D(VertexUVs[TriangleVertexInstanceIDs[TriIndex + 1]]),
-					FVector2D(VertexUVs[TriangleVertexInstanceIDs[TriIndex + 2]])
-				};
+					FVector2D TriangleUVs[3] =
+					{
+						FVector2D(VertexUVs[TriangleVertexInstanceIDs[TriIndex]]),
+						FVector2D(VertexUVs[TriangleVertexInstanceIDs[TriIndex + 1]]),
+						FVector2D(VertexUVs[TriangleVertexInstanceIDs[TriIndex + 2]])
+					};
+
+					if (TriangleUVs[0].ContainsNaN() ||
+						TriangleUVs[1].ContainsNaN() ||
+						TriangleUVs[2].ContainsNaN())
+					{
+						UE_CLOG(DebugName != nullptr, LogStaticMeshOperations, Warning, TEXT("Static Mesh %s has NaNs in it's vertex uvs! Triangle index %d -- using identity for tangent basis."), DebugName, StartIndex);
+						TriangleNormals[StartIndex] = FVector3f(1, 0, 0);
+						TriangleTangents[StartIndex] = FVector3f(0, 1, 0);
+						TriangleBinormals[StartIndex] = FVector3f(0, 0, 1);
+						continue;
+					}
 
-				if (TriangleUVs[0].ContainsNaN() ||
-					TriangleUVs[1].ContainsNaN() ||
-					TriangleUVs[2].ContainsNaN())
+					Result = GetTriangleTangentsAndNormalsWithUV(ComparisonThreshold, TriangleVertexPositions, TriangleUVs);
+				}
+				else
 				{
-					UE_LOG(LogStaticMeshOperations, Warning, TEXT("Static Mesh %s has NaNs in it's vertex uvs! Triangle index %d -- using identity for tangent basis."), DebugName ? DebugName : TEXT("<null>"), StartIndex);
-					TriangleNormals[StartIndex] = FVector3f(1, 0, 0);
-					TriangleTangents[StartIndex] = FVector3f(0, 1, 0);
-					TriangleBinormals[StartIndex] = FVector3f(0, 0, 1);
-					continue;
+					Result = GetTriangleTangentsAndNormalsWithNoUVs(ComparisonThreshold, TriangleVertexPositions);
 				}
-
-				TTuple<FVector3f, FVector3f, FVector3f> Result = GetTriangleTangentsAndNormals(ComparisonThreshold, TriangleVertexPositions, TriangleUVs);
 				TriangleNormals[StartIndex] = Result.Get<0>();
 				TriangleTangents[StartIndex] = Result.Get<1>();
 				TriangleBinormals[StartIndex] = Result.Get<2>();
-			}
+		}
 		}
 	);
 }
@@ -1451,7 +1507,7 @@ void FStaticMeshOperations::ComputeTangentsAndNormals(FMeshDescription& MeshDesc
 		{
 			FStaticMeshAttributes Attributes(MeshDescription);
 
-			TArrayView<const FVector2f> VertexUVs = Attributes.GetVertexInstanceUVs().GetRawArray(0);	// Use UV0
+			TArrayView<const FVector2f> VertexUVs;
 			TArrayView<const FVector3f> TriangleNormals = Attributes.GetTriangleNormals().GetRawArray();
 			TArrayView<const FVector3f> TriangleTangents = Attributes.GetTriangleTangents().GetRawArray();
 			TArrayView<const FVector3f> TriangleBinormals = Attributes.GetTriangleBinormals().GetRawArray();
@@ -1461,6 +1517,14 @@ void FStaticMeshOperations::ComputeTangentsAndNormals(FMeshDescription& MeshDesc
 			TArrayView<FVector3f> VertexTangents = Attributes.GetVertexInstanceTangents().GetRawArray();
 			TArrayView<float> VertexBinormalSigns = Attributes.GetVertexInstanceBinormalSigns().GetRawArray();
 
+			// If the mesh has no UVs, average all tangents/bi-normals for a given vertex, rather than try to maintain
+			// the UV flow.
+			if (Attributes.GetVertexInstanceUVs().GetNumChannels() > 0)
+			{
+				// Use UV0 as the base. Same as with ComputeTriangleTangentsAndNormals 
+				VertexUVs = Attributes.GetVertexInstanceUVs().GetRawArray(0);
+			}
+			
 			check(TriangleNormals.Num() > 0);
 			check(TriangleTangents.Num() > 0);
 			check(TriangleBinormals.Num() > 0);
@@ -1500,7 +1564,10 @@ void FStaticMeshOperations::ComputeTangentsAndNormals(FMeshDescription& MeshDesc
 								if (MeshDescription.GetVertexInstanceVertex(VertexInstanceID) == VertexID)
 								{
 									VertexInfo.VertexInstanceID = VertexInstanceID;
-									VertexInfo.UVs = VertexUVs[VertexInstanceID];	// UV0
+									if (!VertexUVs.IsEmpty())
+									{
+										VertexInfo.UVs = VertexUVs[VertexInstanceID];	// UV0
+									}
 									bPointHasAllTangents &= !VertexNormals[VertexInstanceID].IsNearlyZero() && !VertexTangents[VertexInstanceID].IsNearlyZero();
 									if (bPointHasAllTangents)
 									{
@@ -1579,6 +1646,7 @@ void FStaticMeshOperations::ComputeTangentsAndNormals(FMeshDescription& MeshDesc
 					VertexInstanceInGroup.Reset();
 
 					FVector3f GroupNormal(FVector3f::ZeroVector);
+					
 					for (const FTriangleID& TriangleID : Group)
 					{
 						FVertexInfo& CurrentVertexInfo = VertexInfoMap.FindOrAdd(TriangleID);
@@ -1626,11 +1694,11 @@ void FStaticMeshOperations::ComputeTangentsAndNormals(FMeshDescription& MeshDesc
 					GroupNormal.Normalize();
 					if (!bComputeTangentWithMikkTSpace)
 					{
-						for (auto Kvp : GroupTangent)
+						for (auto& Kvp : GroupTangent)
 						{
 							Kvp.Value.Normalize();
 						}
-						for (auto Kvp : GroupBiNormal)
+						for (auto& Kvp : GroupBiNormal)
 						{
 							Kvp.Value.Normalize();
 						}
@@ -1638,7 +1706,7 @@ void FStaticMeshOperations::ComputeTangentsAndNormals(FMeshDescription& MeshDesc
 					//Apply the average NTB on all Vertex instance
 					for (const FVertexInstanceID& VertexInstanceID : VertexInstanceInGroup)
 					{
-						const FVector2f& VertexUV = VertexUVs[VertexInstanceID];	// UV0
+						const FVector2f& VertexUV = !VertexUVs.IsEmpty() ? VertexUVs[VertexInstanceID] : FVector2f::ZeroVector;
 
 						if (VertexNormals[VertexInstanceID].IsNearlyZero(SMALL_NUMBER))
 						{
@@ -3069,6 +3137,7 @@ bool FStaticMeshOperations::ValidateAndFixData(FMeshDescription& MeshDescription
 	TVertexInstanceAttributesRef<float> VertexInstanceBinormalSigns = Attributes.GetVertexInstanceBinormalSigns();
 	TVertexInstanceAttributesRef<FVector2f> VertexInstanceUVs = Attributes.GetVertexInstanceUVs();
 	TVertexInstanceAttributesRef<FVector4f> VertexInstanceColors = Attributes.GetVertexInstanceColors();
+	const int32 NumUVs = VertexInstanceUVs.GetNumChannels();
 	for (const FVertexInstanceID VertexInstanceID : MeshDescription.VertexInstances().GetElementIDs())
 	{
 		if (VertexInstanceNormals[VertexInstanceID].ContainsNaN())
@@ -3086,10 +3155,14 @@ bool FStaticMeshOperations::ValidateAndFixData(FMeshDescription& MeshDescription
 			bHasInvalidTangentSpaces = true;
 			VertexInstanceBinormalSigns[VertexInstanceID] = 0.0f;
 		}
-		if (VertexInstanceUVs[VertexInstanceID].ContainsNaN())
+
+		for (int32 UVIndex = 0; UVIndex < NumUVs; UVIndex++)
 		{
-			bHasInvalidUVs = true;
-			VertexInstanceUVs[VertexInstanceID] = FVector2f::Zero();
+			if (VertexInstanceUVs.Get(VertexInstanceID, UVIndex).ContainsNaN())
+			{
+				bHasInvalidUVs = true;
+				VertexInstanceUVs.Set(VertexInstanceID, UVIndex, FVector2f::Zero());
+			}
 		}
 		if (VertexInstanceColors[VertexInstanceID].ContainsNaN())
 		{
@@ -3098,21 +3171,24 @@ bool FStaticMeshOperations::ValidateAndFixData(FMeshDescription& MeshDescription
 		}
 	}
 
-	if (bHasInvalidPositions)
-	{
-		UE_LOG(LogStaticMeshOperations, Display, TEXT("Mesh %s has NaNs in it's vertex positions! Offending positions are set to zero."), *DebugName);
-	}
-	if (bHasInvalidTangentSpaces)
-	{
-		UE_LOG(LogStaticMeshOperations, Display, TEXT("Mesh %s has NaNs in it's vertex instance tangent space! Offending tangents are set to zero."), *DebugName);
-	}
-	if (bHasInvalidUVs)
-	{
-		UE_LOG(LogStaticMeshOperations, Display, TEXT("Mesh %s has NaNs in it's vertex instance uvs! Offending uvs are set to zero."), *DebugName);
-	}
-	if (bHasInvalidVertexColors)
+	if (!DebugName.IsEmpty())
 	{
-		UE_LOG(LogStaticMeshOperations, Display, TEXT("Mesh %s has NaNs in it's vertex instance colors! Offending colors are set to white."), *DebugName);
+		if (bHasInvalidPositions)
+		{
+			UE_LOG(LogStaticMeshOperations, Display, TEXT("Mesh %s has NaNs in it's vertex positions! Offending positions are set to zero."), *DebugName);
+		}
+		if (bHasInvalidTangentSpaces)
+		{
+			UE_LOG(LogStaticMeshOperations, Display, TEXT("Mesh %s has NaNs in it's vertex instance tangent space! Offending tangents are set to zero."), *DebugName);
+		}
+		if (bHasInvalidUVs)
+		{
+			UE_LOG(LogStaticMeshOperations, Display, TEXT("Mesh %s has NaNs in it's vertex instance uvs! Offending uvs are set to zero."), *DebugName);
+		}
+		if (bHasInvalidVertexColors)
+		{
+			UE_LOG(LogStaticMeshOperations, Display, TEXT("Mesh %s has NaNs in it's vertex instance colors! Offending colors are set to white."), *DebugName);
+		}
 	}
 
 	return !bHasInvalidPositions && !bHasInvalidTangentSpaces && !bHasInvalidUVs && !bHasInvalidVertexColors;
