Skeletal Mesh: Fix checking of invalid normals.

#rnx

diff --git a/Engine/Plugins/Experimental/Enterprise/DatasmithRuntime/Source/Private/MeshImporter.cpp b/Engine/Plugins/Experimental/Enterprise/DatasmithRuntime/Source/Private/MeshImporter.cpp
--- a/Engine/Plugins/Experimental/Enterprise/DatasmithRuntime/Source/Private/MeshImporter.cpp
+++ b/Engine/Plugins/Experimental/Enterprise/DatasmithRuntime/Source/Private/MeshImporter.cpp
@@ -386,7 +386,7 @@ namespace DatasmithRuntime
 			// We should always have valid normals, tangents and binormals
 			bool bHasInvalidNormals;
 			bool bHasInvalidTangents;
-			FStaticMeshOperations::AreNormalsAndTangentsValid(MeshDescription, bHasInvalidNormals, bHasInvalidTangents);
+			FStaticMeshOperations::HasInvalidVertexInstanceNormalsOrTangents(MeshDescription, bHasInvalidNormals, bHasInvalidTangents);
 
 			// If normals are invalid, compute normals and tangents at polygon level then vertex level
 			if (bHasInvalidNormals)
diff --git a/Engine/Plugins/Experimental/PlanarCutPlugin/Source/PlanarCut/Private/PlanarCut.cpp b/Engine/Plugins/Experimental/PlanarCutPlugin/Source/PlanarCut/Private/PlanarCut.cpp
--- a/Engine/Plugins/Experimental/PlanarCutPlugin/Source/PlanarCut/Private/PlanarCut.cpp
+++ b/Engine/Plugins/Experimental/PlanarCutPlugin/Source/PlanarCut/Private/PlanarCut.cpp
@@ -1034,7 +1034,7 @@ FDynamicMesh3 ConvertMeshDescriptionToCuttingDynamicMesh(const FMeshDescription*
 	FDynamicMesh3 FullMesh; // full-featured conversion of the source mesh
 	Converter.Convert(CuttingMesh, FullMesh, true);
 	bool bHasInvalidNormals, bHasInvalidTangents;
-	FStaticMeshOperations::AreNormalsAndTangentsValid(*CuttingMesh, bHasInvalidNormals, bHasInvalidTangents);
+	FStaticMeshOperations::HasInvalidVertexInstanceNormalsOrTangents(*CuttingMesh, bHasInvalidNormals, bHasInvalidTangents);
 	if (bHasInvalidNormals || bHasInvalidTangents)
 	{
 		FDynamicMeshAttributeSet& Attribs = *FullMesh.Attributes();
diff --git a/Engine/Plugins/Runtime/HairStrands/Source/HairStrandsCore/Private/HairCardsBuilder.cpp b/Engine/Plugins/Runtime/HairStrands/Source/HairStrandsCore/Private/HairCardsBuilder.cpp
--- a/Engine/Plugins/Runtime/HairStrands/Source/HairStrandsCore/Private/HairCardsBuilder.cpp
+++ b/Engine/Plugins/Runtime/HairStrands/Source/HairStrandsCore/Private/HairCardsBuilder.cpp
@@ -417,7 +417,7 @@ void SanitizeMeshDescription(FMeshDescription* MeshDescription)
 
 	bool bHasInvalidNormals = false;
 	bool bHasInvalidTangents = false;
-	FStaticMeshOperations::AreNormalsAndTangentsValid(*MeshDescription, bHasInvalidNormals, bHasInvalidTangents);
+	FStaticMeshOperations::HasInvalidVertexInstanceNormalsOrTangents(*MeshDescription, bHasInvalidNormals, bHasInvalidTangents);
 	if (!bHasInvalidNormals || !bHasInvalidTangents)
 	{
 		FStaticMeshOperations::ComputeTriangleTangentsAndNormals(*MeshDescription, THRESH_POINTS_ARE_SAME);
diff --git a/Engine/Source/Runtime/Engine/Private/Rendering/SkeletalMeshLODImporterData.cpp b/Engine/Source/Runtime/Engine/Private/Rendering/SkeletalMeshLODImporterData.cpp
--- a/Engine/Source/Runtime/Engine/Private/Rendering/SkeletalMeshLODImporterData.cpp
+++ b/Engine/Source/Runtime/Engine/Private/Rendering/SkeletalMeshLODImporterData.cpp
@@ -2376,9 +2376,9 @@ bool FSkeletalMeshImportData::GetMeshDescription(const USkeletalMesh* InSkeletal
 	// Check if we have any broken data, including UVs and normals/tangents.
 	FSkeletalMeshOperations::ValidateAndFixData(OutMeshDescription, SkeletalMeshPath);
 
-	bool bNormalsValid, bTangentsValid;
-	FStaticMeshOperations::AreNormalsAndTangentsValid(OutMeshDescription, bNormalsValid, bTangentsValid);
-	if (!bNormalsValid || !bTangentsValid)
+	bool bHasInvalidNormals, bHasInvalidTangents;
+	FStaticMeshOperations::HasInvalidVertexInstanceNormalsOrTangents(OutMeshDescription, bHasInvalidNormals, bHasInvalidTangents);
+	if (bHasInvalidNormals || bHasInvalidTangents)
 	{
 		// This is required by FSkeletalMeshOperations::ComputeTangentsAndNormals to function correctly.
 		FSkeletalMeshOperations::ComputeTriangleTangentsAndNormals(OutMeshDescription, UE_SMALL_NUMBER, !SkeletalMeshPath.IsEmpty() ? *SkeletalMeshPath : nullptr);
diff --git a/Engine/Source/Runtime/StaticMeshDescription/Private/StaticMeshOperations.cpp b/Engine/Source/Runtime/StaticMeshDescription/Private/StaticMeshOperations.cpp
--- a/Engine/Source/Runtime/StaticMeshDescription/Private/StaticMeshOperations.cpp
+++ b/Engine/Source/Runtime/StaticMeshDescription/Private/StaticMeshOperations.cpp
@@ -1302,7 +1302,7 @@ void FStaticMeshOperations::AppendMeshDescription(const FMeshDescription& Source
 
 //////////////////////////////////////////////////////////////////////////
 // Normals tangents and Bi-normals
-void FStaticMeshOperations::AreNormalsAndTangentsValid(const FMeshDescription& MeshDescription, bool& bHasInvalidNormals, bool& bHasInvalidTangents)
+void FStaticMeshOperations::HasInvalidVertexInstanceNormalsOrTangents(const FMeshDescription& MeshDescription, bool& bHasInvalidNormals, bool& bHasInvalidTangents)
 {
 	bHasInvalidNormals = false;
 	bHasInvalidTangents = false;
@@ -1369,7 +1369,7 @@ void FStaticMeshOperations::RecomputeNormalsAndTangentsIfNeeded(FMeshDescription
 		bool bRecomputeNormals = false;
 		bool bRecomputeTangents = false;
 		
-		AreNormalsAndTangentsValid(MeshDescription, bRecomputeNormals, bRecomputeTangents);
+		HasInvalidVertexInstanceNormalsOrTangents(MeshDescription, bRecomputeNormals, bRecomputeTangents);
 		
 		ComputeNTBsOptions |= (bRecomputeNormals ? EComputeNTBsFlags::Normals : EComputeNTBsFlags::None);
 		ComputeNTBsOptions |= (bRecomputeTangents ? EComputeNTBsFlags::Tangents : EComputeNTBsFlags::None);
diff --git a/Engine/Source/Runtime/StaticMeshDescription/Public/StaticMeshOperations.h b/Engine/Source/Runtime/StaticMeshDescription/Public/StaticMeshOperations.h
--- a/Engine/Source/Runtime/StaticMeshDescription/Public/StaticMeshOperations.h
+++ b/Engine/Source/Runtime/StaticMeshDescription/Public/StaticMeshOperations.h
@@ -99,7 +99,19 @@ public:
 	static STATICMESHDESCRIPTION_API void AppendMeshDescription(const FMeshDescription& SourceMesh, FMeshDescription& TargetMesh, FAppendSettings& AppendSettings);
 	static STATICMESHDESCRIPTION_API void AppendMeshDescriptions(const TArray<const FMeshDescription*>& SourceMeshes, FMeshDescription& TargetMesh, FAppendSettings& AppendSettings);
 
-	static STATICMESHDESCRIPTION_API void AreNormalsAndTangentsValid(const FMeshDescription& MeshDescription, bool& bHasInvalidNormals, bool& bHasInvalidTangents);
+	/** Check the mesh description's vertex instance normals and tangents to see if any of the values are invalid (e.g. nearly zero or they contain a NaN).
+	 *  \note This does not check if the vectors are normalized (which they should for correct normals and tangents).
+	 *  \param MeshDescription The mesh description object to check.
+	 *  \param bHasInvalidNormals The vertex instance normal attribute contains a near zero-length vector or a NaN component value.
+	 *  \param bHasInvalidTangents The vertex instance tangent attribute contains a near zero-length vector or a NaN component value.
+	 */
+	static STATICMESHDESCRIPTION_API void HasInvalidVertexInstanceNormalsOrTangents(const FMeshDescription& MeshDescription, bool& bHasInvalidNormals, bool& bHasInvalidTangents);
+
+	UE_DEPRECATED(5.5, "Use HasInvalidVertexInstanceNormalsOrTangents")
+	static void AreNormalsAndTangentsValid(const FMeshDescription& MeshDescription, bool& bHasInvalidNormals, bool& bHasInvalidTangents)
+	{
+		HasInvalidVertexInstanceNormalsOrTangents(MeshDescription, bHasInvalidNormals, bHasInvalidTangents);
+	}
 
 
 	/** Find all overlapping vertex using the threshold in the mesh description. */
