Fix static analysis warnings

diff --git a/Engine/Source/Runtime/AssetRegistry/Private/AssetRegistry.cpp b/Engine/Source/Runtime/AssetRegistry/Private/AssetRegistry.cpp
--- a/Engine/Source/Runtime/AssetRegistry/Private/AssetRegistry.cpp
+++ b/Engine/Source/Runtime/AssetRegistry/Private/AssetRegistry.cpp
@@ -1018,8 +1018,11 @@ void FAssetRegistryImpl::Initialize(Impl::FInitializeContext& Context)
 		{
 			// For the Editor and editor game we need to take responsibility for the synchronous search;
 			// Commandlets and cooked game will handle it themselves.
-			constexpr bool bEditorExecutable = WITH_EDITOR;
-			Context.bNeedsSearchAllAssetsAtStartSynchronous = bEditorExecutable && !IsRunningCommandlet();
+#if WITH_EDITOR
+			Context.bNeedsSearchAllAssetsAtStartSynchronous = !IsRunningCommandlet();
+#else
+			Context.bNeedsSearchAllAssetsAtStartSynchronous = false;
+#endif 
 		}
 	}
 
diff --git a/Engine/Source/Runtime/StaticMeshDescription/Private/StaticMeshOperations.cpp b/Engine/Source/Runtime/StaticMeshDescription/Private/StaticMeshOperations.cpp
--- a/Engine/Source/Runtime/StaticMeshDescription/Private/StaticMeshOperations.cpp
+++ b/Engine/Source/Runtime/StaticMeshDescription/Private/StaticMeshOperations.cpp
@@ -2261,10 +2261,12 @@ bool FStaticMeshOperations::GenerateUV(const FMeshDescription& MeshDescription,
 
 	OutTexCoords.Reset();
 
-	const bool bAutoUVAvailable = WITH_EDITOR;
 	const bool bHasUVs = VertexInstanceUVs.GetNumElements() > 0;
-	const bool bUseLegacy = GenerateUVOptions.UVMethod == EGenerateUVMethod::Legacy || !bAutoUVAvailable;
-	if (bHasUVs && bUseLegacy)
+#if WITH_EDITOR
+	if (bHasUVs && GenerateUVOptions.UVMethod == EGenerateUVMethod::Legacy)
+#else
+	if (bHasUVs)
+#endif
 	{
 		FUniqueUVMeshDescriptionView MeshDescriptionView(MeshDescription, GenerateUVOptions.bMergeTrianglesWithIdenticalAttributes, OutTexCoords);
 
