Fixes static mesh to skeletal mesh conversion losing multiple UV channels

diff --git a/Engine/Source/Runtime/StaticMeshDescription/Private/StaticMeshAttributes.cpp b/Engine/Source/Runtime/StaticMeshDescription/Private/StaticMeshAttributes.cpp
--- a/Engine/Source/Runtime/StaticMeshDescription/Private/StaticMeshAttributes.cpp
+++ b/Engine/Source/Runtime/StaticMeshDescription/Private/StaticMeshAttributes.cpp
@@ -42,7 +42,8 @@ void FStaticMeshAttributes::Register(bool bKeepExistingAttribute)
 	// Add basic vertex instance attributes
 	if (!MeshDescription.VertexInstanceAttributes().HasAttribute(MeshAttribute::VertexInstance::TextureCoordinate) || !bKeepExistingAttribute)
 	{
-		MeshDescription.VertexInstanceAttributes().RegisterAttribute<FVector2f>(MeshAttribute::VertexInstance::TextureCoordinate, 1, FVector2f::ZeroVector, EMeshAttributeFlags::Lerpable | EMeshAttributeFlags::Mandatory);
+		const int32 NumUVChannels = MeshDescription.GetNumUVElementChannels();
+		MeshDescription.VertexInstanceAttributes().RegisterAttribute<FVector2f>(MeshAttribute::VertexInstance::TextureCoordinate, NumUVChannels, FVector2f::ZeroVector, EMeshAttributeFlags::Lerpable | EMeshAttributeFlags::Mandatory);
 	}
 
 	if (!MeshDescription.VertexInstanceAttributes().HasAttribute(MeshAttribute::VertexInstance::Normal) || !bKeepExistingAttribute)
