[CAD] Fixed some file import hanging the editor

diff --git a/Engine/Source/Runtime/StaticMeshDescription/Private/StaticMeshOperations.cpp b/Engine/Source/Runtime/StaticMeshDescription/Private/StaticMeshOperations.cpp
--- a/Engine/Source/Runtime/StaticMeshDescription/Private/StaticMeshOperations.cpp
+++ b/Engine/Source/Runtime/StaticMeshDescription/Private/StaticMeshOperations.cpp
@@ -1534,7 +1534,7 @@ void FStaticMeshOperations::ComputeTangentsAndNormals(FMeshDescription& MeshDesc
 			TMap<FVector2f, FVector3f, FDefaultSetAllocator, FNTBGroupKeyFuncs> GroupBiNormal;
 			TMap<FTriangleID, FVertexInfo> VertexInfoMap;
 			TArray<TArray<FTriangleID, TInlineAllocator<8>>> Groups;
-			TArray<FTriangleID> ConsumedTriangle;
+			TSet<FTriangleID> ConsumedTriangle;
 			TArray<FTriangleID> PolygonQueue;
 			TArray<FVertexInstanceID> VertexInstanceInGroup;
 
@@ -1596,6 +1596,7 @@ void FStaticMeshOperations::ComputeTangentsAndNormals(FMeshDescription& MeshDesc
 				//Build all group by recursively traverse all polygon connected to the vertex
 				Groups.Reset();
 				ConsumedTriangle.Reset();
+				TSet<FEdgeID> ConsumedEdges;
 				for (auto Kvp : VertexInfoMap)
 				{
 					if (ConsumedTriangle.Contains(Kvp.Key))
@@ -1612,14 +1613,17 @@ void FStaticMeshOperations::ComputeTangentsAndNormals(FMeshDescription& MeshDesc
 						FTriangleID CurrentPolygonID = PolygonQueue.Pop(EAllowShrinking::No);
 						FVertexInfo& CurrentVertexInfo = VertexInfoMap.FindOrAdd(CurrentPolygonID);
 						CurrentGroup.AddUnique(CurrentVertexInfo.TriangleID);
-						ConsumedTriangle.AddUnique(CurrentVertexInfo.TriangleID);
+						ConsumedTriangle.Add(CurrentVertexInfo.TriangleID);
 						for (const FEdgeID& EdgeID : CurrentVertexInfo.EdgeIDs)
 						{
-							if (EdgeHardnesses[EdgeID])
+							if (EdgeHardnesses[EdgeID] || ConsumedEdges.Contains(EdgeID))
 							{
-								//End of the group
+								//End of the group or non manifold edge
 								continue;
 							}
+
+							ConsumedEdges.Add(EdgeID);
+
 							for (const FTriangleID& TriangleID : MeshDescription.GetEdgeConnectedTriangleIDs(EdgeID))
 							{
 								if (TriangleID == CurrentVertexInfo.TriangleID)
