Mesh Description: Polygon attribute for storing names for mesh parts.

- This allows identifying parts of the unified mesh that originated from separate objects on import, or during construction in modeling/geometry scripting.
- Removed deprecated mesh attributes.

#rnx

diff --git a/Engine/Plugins/Interchange/Runtime/Source/Import/Private/Mesh/InterchangeOBJTranslator.cpp b/Engine/Plugins/Interchange/Runtime/Source/Import/Private/Mesh/InterchangeOBJTranslator.cpp
--- a/Engine/Plugins/Interchange/Runtime/Source/Import/Private/Mesh/InterchangeOBJTranslator.cpp
+++ b/Engine/Plugins/Interchange/Runtime/Source/Import/Private/Mesh/InterchangeOBJTranslator.cpp
@@ -307,6 +307,7 @@ FMeshDescription FObjData::MakeMeshDescriptionForGroup(const FString& GroupName,
 
 	FStaticMeshAttributes Attributes(MeshDescription);
 	Attributes.Register();
+	Attributes.RegisterPolygonObjectNameAttribute();
 
 	MeshDescription.SuspendVertexInstanceIndexing();
 	MeshDescription.SuspendEdgeIndexing();
@@ -484,6 +485,15 @@ FMeshDescription FObjData::MakeMeshDescriptionForGroup(const FString& GroupName,
 		}
 	}
 
+	// Fill in the object group name onto all the triangles of the mesh.
+	TPolygonAttributesRef<FName> PolygonObjectNames = Attributes.GetPolygonObjectNames();
+	const FName ObjectName(GroupName);
+	for (const FPolygonID PolygonID: MeshDescription.Polygons().GetElementIDs())
+	{
+		PolygonObjectNames.Set(PolygonID, ObjectName);
+	}
+	
+
 	MeshDescription.ResumeVertexInstanceIndexing();
 	MeshDescription.ResumePolygonIndexing();
 	MeshDescription.ResumePolygonGroupIndexing();
diff --git a/Engine/Plugins/Interchange/Runtime/Source/Import/Private/Mesh/InterchangeSkeletalMeshFactory.cpp b/Engine/Plugins/Interchange/Runtime/Source/Import/Private/Mesh/InterchangeSkeletalMeshFactory.cpp
--- a/Engine/Plugins/Interchange/Runtime/Source/Import/Private/Mesh/InterchangeSkeletalMeshFactory.cpp
+++ b/Engine/Plugins/Interchange/Runtime/Source/Import/Private/Mesh/InterchangeSkeletalMeshFactory.cpp
@@ -390,6 +390,7 @@ namespace UE
 				{
 					AppendSettings.bMergeUVChannels[ChannelIdx] = true;
 				}
+				AppendSettings.bMergeObjectName = true;
 				
 				bool bKeepSectionsSeparate = false;
 				SkeletalMeshFactoryNode->GetCustomKeepSectionsSeparate(bKeepSectionsSeparate);
diff --git a/Engine/Plugins/Interchange/Runtime/Source/Import/Private/Mesh/InterchangeStaticMeshFactory.cpp b/Engine/Plugins/Interchange/Runtime/Source/Import/Private/Mesh/InterchangeStaticMeshFactory.cpp
--- a/Engine/Plugins/Interchange/Runtime/Source/Import/Private/Mesh/InterchangeStaticMeshFactory.cpp
+++ b/Engine/Plugins/Interchange/Runtime/Source/Import/Private/Mesh/InterchangeStaticMeshFactory.cpp
@@ -430,6 +430,9 @@ UInterchangeFactoryBase::FImportAssetResult UInterchangeStaticMeshFactory::Impor
 		{
 			AppendSettings.bMergeUVChannels[ChannelIdx] = true;
 		}
+		
+		// Merge the SourceObjectName attribute, if it exists on both the target and source.
+		AppendSettings.bMergeObjectName = true;
 				
 
 		TArray<FString> MeshUids;
diff --git a/Engine/Plugins/Interchange/Runtime/Source/Parsers/Fbx/Private/FbxMesh.cpp b/Engine/Plugins/Interchange/Runtime/Source/Parsers/Fbx/Private/FbxMesh.cpp
--- a/Engine/Plugins/Interchange/Runtime/Source/Parsers/Fbx/Private/FbxMesh.cpp
+++ b/Engine/Plugins/Interchange/Runtime/Source/Parsers/Fbx/Private/FbxMesh.cpp
@@ -584,6 +584,9 @@ bool FMeshDescriptionImporter::FillMeshDescriptionFromFbxMesh(FbxMesh* Mesh, con
 	FStaticMeshAttributes Attributes(*MeshDescription);
 	constexpr bool bKeepExistingAttribute = true;
 	Attributes.Register(bKeepExistingAttribute);
+	
+	// We want to store the source mesh name with each triangle so that we can identify individual geometry pieces in the combined static/skeletal mesh.
+	Attributes.RegisterPolygonObjectNameAttribute();
 
 	bool bStaticMeshUseSmoothEdgesIfSmoothingInformationIsMissing = true;
 
@@ -601,10 +604,15 @@ bool FMeshDescriptionImporter::FillMeshDescriptionFromFbxMesh(FbxMesh* Mesh, con
 		Message->Text = LOCTEXT("CantGetLayer", "Unable to get FBX base layer for the mesh '{MeshName}'.");
 		return false;
 	}
-
+	
 	FFBXUVs FBXUVs(*this, Mesh);
 
 	FbxNode* MeshNode = Mesh->GetNodeCount() > 0 ? Mesh->GetNode() : nullptr; //Get the first node using the mesh to setup default asset materials
+
+	// Grab the mesh node's name, we use that to fill the object name attribute. The mesh object name itself is of no use since it's the same
+	// for all meshes.
+	const FName MeshNodeName(Parser.GetFbxHelper()->GetFbxObjectName(MeshNode));
+	
 	//
 	// create materials
 	//
@@ -831,6 +839,7 @@ bool FMeshDescriptionImporter::FillMeshDescriptionFromFbxMesh(FbxMesh* Mesh, con
 		TVertexInstanceAttributesRef<FVector2f> VertexInstanceUVs = Attributes.GetVertexInstanceUVs();
 		TEdgeAttributesRef<bool> EdgeHardnesses = Attributes.GetEdgeHardnesses();
 		TPolygonGroupAttributesRef<FName> PolygonGroupImportedMaterialSlotNames = Attributes.GetPolygonGroupMaterialSlotNames();
+		TPolygonAttributesRef<FName> PolygonObjectNames = Attributes.GetPolygonObjectNames();
 
 		int32 VertexOffset = MeshDescription->Vertices().Num();
 		int32 VertexInstanceOffset = MeshDescription->VertexInstances().Num();
@@ -1314,6 +1323,9 @@ bool FMeshDescriptionImporter::FillMeshDescriptionFromFbxMesh(FbxMesh* Mesh, con
 				const FTriangleID NewTriangleID = MeshDescription->CreateTriangle(PolygonGroupID, CornerInstanceIDs, &NewEdgeIDs);
 				check(NewEdgeIDs.Num() == 0);
 
+				// Set the name of the object that generated this triangle.
+				PolygonObjectNames.Set(MeshDescription->GetTrianglePolygon(NewTriangleID), MeshNodeName);
+
 				for (int32 UVLayerIndex = 0; UVLayerIndex < FBXUVs.UniqueUVCount; UVLayerIndex++)
 				{
 					FUVID UVIDs[3] = { INDEX_NONE, INDEX_NONE, INDEX_NONE };
diff --git a/Engine/Plugins/Interchange/Runtime/Source/Parsers/GLTFCore/Private/GLTFMeshFactory.cpp b/Engine/Plugins/Interchange/Runtime/Source/Parsers/GLTFCore/Private/GLTFMeshFactory.cpp
--- a/Engine/Plugins/Interchange/Runtime/Source/Parsers/GLTFCore/Private/GLTFMeshFactory.cpp
+++ b/Engine/Plugins/Interchange/Runtime/Source/Parsers/GLTFCore/Private/GLTFMeshFactory.cpp
@@ -23,6 +23,7 @@ namespace GLTF
 	private:
 
 		bool ImportPrimitive(const GLTF::FPrimitive&                        Primitive,  //
+			FName										   MeshName,
 			int32                                          PrimitiveIndex,
 			int32                                          NumUVs,
 			bool                                           bMeshHasTagents,
@@ -33,6 +34,7 @@ namespace GLTF
 			const TVertexInstanceAttributesRef<FVector2f>& VertexInstanceUVs,
 			const TVertexInstanceAttributesRef<FVector4f>&  VertexInstanceColors,
 			const TEdgeAttributesRef<bool>&                EdgeHardnesses,
+			TPolygonAttributesRef<FName>&				PolygonObjectNames,
 			const FTransform& MeshGlobalTransform,
 			FMeshDescription* MeshDescription,
 			const TArray<float>& MorphTargetWeights);
@@ -311,10 +313,12 @@ namespace GLTF
 
 	void FMeshFactoryImpl::FillMeshDescription(const FMesh &Mesh, const FTransform& MeshGlobalTransform, FMeshDescription* MeshDescription, const TArray<float>& MorphTargetWeights)
 	{
+		const FName MeshName(Mesh.Name);
 		const int32 NumUVs = FMath::Max(1, GetNumUVs(Mesh));
 
 		FStaticMeshAttributes StaticMeshAttributes(*MeshDescription);
 		StaticMeshAttributes.Register();
+		StaticMeshAttributes.RegisterPolygonObjectNameAttribute();
 
 		FMatrix TotalMatrix = MeshGlobalTransform.ToMatrixWithScale();
 
@@ -326,6 +330,8 @@ namespace GLTF
 		TVertexInstanceAttributesRef<float> VertexInstanceBinormalSigns = StaticMeshAttributes.GetVertexInstanceBinormalSigns();
 		TVertexInstanceAttributesRef<FVector2f> VertexInstanceUVs = StaticMeshAttributes.GetVertexInstanceUVs();
 		TVertexInstanceAttributesRef<FVector4f> VertexInstanceColors = StaticMeshAttributes.GetVertexInstanceColors();
+		TPolygonAttributesRef<FName> PolygonObjectNames = StaticMeshAttributes.GetPolygonObjectNames();
+
 		VertexInstanceUVs.SetNumChannels(NumUVs);
 
 		MaterialIndicesUsed.Empty(10);
@@ -399,10 +405,10 @@ namespace GLTF
 			}
 
 			const bool        bHasDegenerateTriangles =
-				ImportPrimitive(Primitive, Index, NumUVs, Mesh.HasTangents(), Mesh.HasColors(),  //
+				ImportPrimitive(Primitive, MeshName, Index, NumUVs, Mesh.HasTangents(), Mesh.HasColors(),  //
 					VertexInstanceNormals, VertexInstanceTangents, VertexInstanceBinormalSigns, VertexInstanceUVs,
 					VertexInstanceColors,  //
-					EdgeHardnesses, MeshGlobalTransform, MeshDescription,
+					EdgeHardnesses, PolygonObjectNames, MeshGlobalTransform, MeshDescription,
 					MorphTargetWeights);
 
 			bMeshUsesEmptyMaterial |= Primitive.MaterialIndex == INDEX_NONE;
@@ -427,6 +433,7 @@ namespace GLTF
 	}
 
 	bool FMeshFactoryImpl::ImportPrimitive(const GLTF::FPrimitive&                        Primitive,  //
+		FName										MeshName,
 		int32                                          PrimitiveIndex,
 		int32                                          NumUVs,
 		bool                                           bMeshHasTangents,
@@ -437,6 +444,7 @@ namespace GLTF
 		const TVertexInstanceAttributesRef<FVector2f>& VertexInstanceUVs,
 		const TVertexInstanceAttributesRef<FVector4f>&  VertexInstanceColors,
 		const TEdgeAttributesRef<bool>&                EdgeHardnesses,
+		TPolygonAttributesRef<FName>&				PolygonObjectNames,
 		const FTransform& MeshGlobalTransform,
 		FMeshDescription* MeshDescription,
 		const TArray<float>& MorphTargetWeights)
@@ -720,6 +728,8 @@ namespace GLTF
 					//   We want to combine identical verts whether they're smooth neighbors or triangles belonging to the same flat polygon.)
 					EdgeHardnesses[NewEdgeID] = false;
 				}
+
+				PolygonObjectNames.Set(NewPolygonID, MeshName);
 			}
 		}
 		
diff --git a/Engine/Source/Runtime/SkeletalMeshDescription/Private/SkeletalMeshOperations.cpp b/Engine/Source/Runtime/SkeletalMeshDescription/Private/SkeletalMeshOperations.cpp
--- a/Engine/Source/Runtime/SkeletalMeshDescription/Private/SkeletalMeshOperations.cpp
+++ b/Engine/Source/Runtime/SkeletalMeshDescription/Private/SkeletalMeshOperations.cpp
@@ -1172,6 +1172,9 @@ void FSkeletalMeshOperations::FixVertexInstanceStructure(FMeshDescription& Sourc
 		TVertexInstanceAttributesRef<FVector4f> Colors;
 		TVertexInstanceAttributesRef<FVector2f> UVs;
 
+		// For read/write of Polygon attributes
+		TPolygonAttributesRef<FName> PolygonObjectNames;  
+
 		//GeometryParts
 		FSkeletalMeshAttributes::FSourceGeometryPartNameRef GeometryPartNames;
 		FSkeletalMeshAttributes::FSourceGeometryPartVertexOffsetAndCountRef GeometryPartVertexOffsetAndCounts;
@@ -1183,14 +1186,27 @@ void FSkeletalMeshOperations::FixVertexInstanceStructure(FMeshDescription& Sourc
 		TArray<TVertexAttributesRef<FVector3f>> MorphPositionDeltas; //Vertex, SourceMeshMorphPosDeltas[0] belongs to -> SourceMorphTargetNames[0] and so on..
 		TArray<TVertexInstanceAttributesRef<FVector3f>> MorphNormals; //VertexInstance, SourceMeshMorphPosDeltas[0] belongs to -> SourceMorphTargetNames[0] and so on..
 
-		FMeshAttributesHelper(FMeshDescription& SourceMeshDescription, bool bRegisterAttributes = false)
+		FMeshAttributesHelper(FMeshDescription& SourceMeshDescription, const FMeshDescription* MeshDescriptionTemplate = nullptr)
 			: MeshAttributes(SourceMeshDescription)
 		{
-			if (bRegisterAttributes)
+			if (MeshDescriptionTemplate)
 			{
 				MeshAttributes.Register();
 
-				MeshAttributes.RegisterSourceGeometryPartsAttributes();
+				FSkeletalMeshConstAttributes TemplateAttributes(*MeshDescriptionTemplate);
+				if (TemplateAttributes.GetPolygonObjectNames().IsValid())
+				{
+					MeshAttributes.RegisterPolygonObjectNameAttribute();
+				}
+				if (TemplateAttributes.HasSourceGeometryParts())
+				{
+					MeshAttributes.RegisterSourceGeometryPartsAttributes();
+				}
+
+				for (const FName& MorphTargetName : TemplateAttributes.GetMorphTargetNames())
+				{
+					MeshAttributes.RegisterMorphTargetAttribute(MorphTargetName, TemplateAttributes.HasMorphTargetNormalsAttribute(MorphTargetName));
+				}
 			}
 
 			PolygonGroupMaterialSlotNames = MeshAttributes.GetPolygonGroupMaterialSlotNames();
@@ -1204,6 +1220,8 @@ void FSkeletalMeshOperations::FixVertexInstanceStructure(FMeshDescription& Sourc
 			Colors = MeshAttributes.GetVertexInstanceColors();
 			UVs = MeshAttributes.GetVertexInstanceUVs();
 
+			PolygonObjectNames = MeshAttributes.GetPolygonObjectNames();
+
 			GeometryPartNames = MeshAttributes.GetSourceGeometryPartNames();
 			GeometryPartVertexOffsetAndCounts = MeshAttributes.GetSourceGeometryPartVertexOffsetAndCounts();
 
@@ -1223,7 +1241,8 @@ void FSkeletalMeshOperations::FixVertexInstanceStructure(FMeshDescription& Sourc
 	};
 
 	FMeshAttributesHelper Source(SourceMeshDescription);
-	FMeshAttributesHelper Target(TargetMeshDescription, true);
+
+	FMeshAttributesHelper Target(TargetMeshDescription, &SourceMeshDescription);
 
 	const int32 TriangleCount = SourceMeshDescription.Triangles().Num();
 	const int32 VertexInstanceCount = TriangleCount * 3;
@@ -1243,7 +1262,7 @@ void FSkeletalMeshOperations::FixVertexInstanceStructure(FMeshDescription& Sourc
 	for (size_t PolygonGroupIndex = 0; PolygonGroupIndex < Source.PolygonGroupMaterialSlotNames.GetNumElements(); PolygonGroupIndex++)
 	{
 		FPolygonGroupID PolygonGroupID = TargetMeshDescription.CreatePolygonGroup();
-		Target.PolygonGroupMaterialSlotNames[PolygonGroupIndex] = Source.PolygonGroupMaterialSlotNames[PolygonGroupIndex];
+		Target.PolygonGroupMaterialSlotNames[PolygonGroupID] = Source.PolygonGroupMaterialSlotNames[PolygonGroupIndex];
 	}
 
 	//Copy vertices (aka Position and SkinWeights)
@@ -1308,17 +1327,21 @@ void FSkeletalMeshOperations::FixVertexInstanceStructure(FMeshDescription& Sourc
 
 			TargetVertexInstanceIDs[Corner] = TargetVertexInstanceID;
 			SourceToTargetVertexInstanceIDs.Add(TPair<int32, int32>(SourceVertexInstanceID, TargetVertexInstanceID)); //to be used for MorphTargetNormals copying
-				
 		}
 
-		FPolygonID TargetPolyongId = TargetMeshDescription.CreatePolygon(PolygonGroupID, TargetVertexInstanceIDs);
+		FPolygonID TargetPolygonId = TargetMeshDescription.CreatePolygon(PolygonGroupID, TargetVertexInstanceIDs);
+
+		if (Target.PolygonObjectNames.IsValid())
+		{
+			Target.PolygonObjectNames.Set(TargetPolygonId, Source.PolygonObjectNames.Get(PolygonID));
+		}
 
 		//Smooth mask:
 		if (SourceSmoothingMasks.IsValidIndex(TriangleID.GetValue()))
 		{
-			if (ensure(TargetFaceSmoothingMasks.IsValidIndex(TargetPolyongId.GetValue())))
+			if (ensure(TargetFaceSmoothingMasks.IsValidIndex(TargetPolygonId.GetValue())))
 			{
-				TargetFaceSmoothingMasks[TargetPolyongId.GetValue()] = SourceSmoothingMasks[TriangleID.GetValue()];
+				TargetFaceSmoothingMasks[TargetPolygonId.GetValue()] = SourceSmoothingMasks[TriangleID.GetValue()];
 			}
 		}
 	}
@@ -1330,7 +1353,6 @@ void FSkeletalMeshOperations::FixVertexInstanceStructure(FMeshDescription& Sourc
 	UseMorphTargetNormals.SetNumUninitialized(Source.NumberOfMorphTargets);
 	for (const FName& SourceMorphTargetName : Source.MorphTargetNames)
 	{
-		Target.MeshAttributes.RegisterMorphTargetAttribute(SourceMorphTargetName, Source.MeshAttributes.HasMorphTargetNormalsAttribute(SourceMorphTargetName));
 		Target.MorphPositionDeltas.Add(Target.MeshAttributes.GetVertexMorphPositionDelta(SourceMorphTargetName));
 		Target.MorphNormals.Add(Target.MeshAttributes.GetVertexInstanceMorphNormalDelta(SourceMorphTargetName));
 
diff --git a/Engine/Source/Runtime/StaticMeshDescription/Private/StaticMeshAttributes.cpp b/Engine/Source/Runtime/StaticMeshDescription/Private/StaticMeshAttributes.cpp
--- a/Engine/Source/Runtime/StaticMeshDescription/Private/StaticMeshAttributes.cpp
+++ b/Engine/Source/Runtime/StaticMeshDescription/Private/StaticMeshAttributes.cpp
@@ -16,21 +16,10 @@ namespace MeshAttribute
 	const FName Triangle::Tangent("Tangent");
 	const FName Triangle::Binormal("Binormal");
 
-	const FName Polygon::Normal("Normal");
-	const FName Polygon::Tangent("Tangent");
-	const FName Polygon::Binormal("Binormal");
-	const FName Polygon::Center("Center");
+	const FName Polygon::ObjectName("ObjectName");
 
 	const FName PolygonGroup::ImportedMaterialSlotName("ImportedMaterialSlotName");
 
-PRAGMA_DISABLE_DEPRECATION_WARNINGS
-	// all deprected in 4.26
-	const FName Vertex::CornerSharpness("CornerSharpness");
-	const FName PolygonGroup::CastShadow("CastShadow");
-	const FName Edge::CreaseSharpness("CreaseSharpness");
-	const FName Edge::IsUVSeam("IsUVSeam");
-	const FName PolygonGroup::EnableCollision("EnableCollision");
-PRAGMA_ENABLE_DEPRECATION_WARNINGS
 }
 
 
@@ -85,20 +74,18 @@ void FStaticMeshAttributes::Register(bool bKeepExistingAttribute)
 }
 
 
-void FStaticMeshAttributes::RegisterPolygonNormalAndTangentAttributes()
-{
-	PRAGMA_DISABLE_DEPRECATION_WARNINGS
-	MeshDescription.PolygonAttributes().RegisterAttribute<FVector3f>(MeshAttribute::Polygon::Normal, 1, FVector3f::ZeroVector, EMeshAttributeFlags::Transient);
-	MeshDescription.PolygonAttributes().RegisterAttribute<FVector3f>(MeshAttribute::Polygon::Tangent, 1, FVector3f::ZeroVector, EMeshAttributeFlags::Transient);
-	MeshDescription.PolygonAttributes().RegisterAttribute<FVector3f>(MeshAttribute::Polygon::Binormal, 1, FVector3f::ZeroVector, EMeshAttributeFlags::Transient);
-	MeshDescription.PolygonAttributes().RegisterAttribute<FVector3f>(MeshAttribute::Polygon::Center, 1, FVector3f::ZeroVector, EMeshAttributeFlags::Transient);
-	PRAGMA_ENABLE_DEPRECATION_WARNINGS
-}
-
-
 void FStaticMeshAttributes::RegisterTriangleNormalAndTangentAttributes()
 {
 	MeshDescription.TriangleAttributes().RegisterAttribute<FVector3f>(MeshAttribute::Triangle::Normal, 1, FVector3f::ZeroVector, EMeshAttributeFlags::Transient);
 	MeshDescription.TriangleAttributes().RegisterAttribute<FVector3f>(MeshAttribute::Triangle::Tangent, 1, FVector3f::ZeroVector, EMeshAttributeFlags::Transient);
 	MeshDescription.TriangleAttributes().RegisterAttribute<FVector3f>(MeshAttribute::Triangle::Binormal, 1, FVector3f::ZeroVector, EMeshAttributeFlags::Transient);
 }
+
+
+void FStaticMeshAttributes::RegisterPolygonObjectNameAttribute()
+{
+	if (!MeshDescription.PolygonAttributes().HasAttribute(MeshAttribute::Polygon::ObjectName))
+	{
+		MeshDescription.PolygonAttributes().RegisterAttribute<FName>(MeshAttribute::Polygon::ObjectName, 1, NAME_None, EMeshAttributeFlags::AutoGenerated);
+	}
+}
diff --git a/Engine/Source/Runtime/StaticMeshDescription/Private/StaticMeshOperations.cpp b/Engine/Source/Runtime/StaticMeshDescription/Private/StaticMeshOperations.cpp
--- a/Engine/Source/Runtime/StaticMeshDescription/Private/StaticMeshOperations.cpp
+++ b/Engine/Source/Runtime/StaticMeshDescription/Private/StaticMeshOperations.cpp
@@ -26,165 +26,6 @@ DEFINE_LOG_CATEGORY(LogStaticMeshOperations);
 
 #define LOCTEXT_NAMESPACE "StaticMeshOperations"
 
-static bool GetPolygonTangentsAndNormals(FMeshDescription& MeshDescription,
-										 FPolygonID PolygonID,
-										 float ComparisonThreshold, 
-										 TVertexAttributesConstRef<const FVector3f> VertexPositions,
-										 TVertexInstanceAttributesConstRef<const FVector2f> VertexUVs,
-										 TPolygonAttributesRef<FVector3f> PolygonNormals,
-										 TPolygonAttributesRef<FVector3f> PolygonTangents,
-										 TPolygonAttributesRef<FVector3f> PolygonBinormals,
-										 TPolygonAttributesRef<FVector3f> PolygonCenters)
-{
-	bool bValidNTBs = true;
-
-	// Calculate the tangent basis for the polygon, based on the average of all constituent triangles
-	FVector3f Normal(FVector3f::ZeroVector);
-	FVector3f Tangent(FVector3f::ZeroVector);
-	FVector3f Binormal(FVector3f::ZeroVector);
-	FVector3f Center(FVector3f::ZeroVector);
-
-	// Calculate the center of this polygon
-	TArray<FVertexInstanceID, TInlineAllocator<4>> VertexInstanceIDs = MeshDescription.GetPolygonVertexInstances<TInlineAllocator<4>>(PolygonID);
-	for (const FVertexInstanceID& VertexInstanceID : VertexInstanceIDs)
-	{
-		Center += VertexPositions[MeshDescription.GetVertexInstanceVertex(VertexInstanceID)];
-	}
-	Center /= float(VertexInstanceIDs.Num());
-
-	// GetSafeNormal compare the squareSum to the tolerance.
-	const float SquareComparisonThreshold = FMath::Max(ComparisonThreshold * ComparisonThreshold, MIN_flt);
-	for (const FTriangleID& TriangleID : MeshDescription.GetPolygonTriangles(PolygonID))
-	{
-		TArrayView<const FVertexInstanceID> TriangleVertexInstances = MeshDescription.GetTriangleVertexInstances(TriangleID);
-		const FVertexID VertexID0 = MeshDescription.GetVertexInstanceVertex(TriangleVertexInstances[0]);
-		const FVertexID VertexID1 = MeshDescription.GetVertexInstanceVertex(TriangleVertexInstances[1]);
-		const FVertexID VertexID2 = MeshDescription.GetVertexInstanceVertex(TriangleVertexInstances[2]);
-
-		const FVector3f Position0 = VertexPositions[VertexID0];
-		const FVector3f DPosition1 = VertexPositions[VertexID1] - Position0;
-		const FVector3f DPosition2 = VertexPositions[VertexID2] - Position0;
-
-		const FVector2f UV0 = VertexUVs[TriangleVertexInstances[0]];
-		const FVector2f DUV1 = VertexUVs[TriangleVertexInstances[1]] - UV0;
-		const FVector2f DUV2 = VertexUVs[TriangleVertexInstances[2]] - UV0;
-
-		// We have a left-handed coordinate system, but a counter-clockwise winding order
-		// Hence normal calculation has to take the triangle vectors cross product in reverse.
-		FVector3f TmpNormal = FVector3f::CrossProduct(DPosition2, DPosition1).GetSafeNormal(SquareComparisonThreshold);
-		if (!TmpNormal.IsNearlyZero(ComparisonThreshold))
-		{
-			FMatrix44f	ParameterToLocal(
-				DPosition1,
-				DPosition2,
-				Position0,
-				FVector3f::ZeroVector
-			);
-
-			FMatrix44f ParameterToTexture(
-				FPlane4f(DUV1.X, DUV1.Y, 0, 0),
-				FPlane4f(DUV2.X, DUV2.Y, 0, 0),
-				FPlane4f(UV0.X, UV0.Y, 1, 0),
-				FPlane4f(0, 0, 0, 1)
-			);
-
-			// Use InverseSlow to catch singular matrices.  Inverse can miss this sometimes.
-			const FMatrix44f TextureToLocal = ParameterToTexture.Inverse() * ParameterToLocal;
-
-			FVector3f TmpTangent = TextureToLocal.TransformVector(FVector3f(1, 0, 0)).GetSafeNormal();
-			FVector3f TmpBinormal = TextureToLocal.TransformVector(FVector3f(0, 1, 0)).GetSafeNormal();
-			FVector3f::CreateOrthonormalBasis(TmpTangent, TmpBinormal, TmpNormal);
-
-			if (TmpTangent.IsNearlyZero() || TmpTangent.ContainsNaN()
-				|| TmpBinormal.IsNearlyZero() || TmpBinormal.ContainsNaN())
-			{
-				TmpTangent = FVector3f::ZeroVector;
-				TmpBinormal = FVector3f::ZeroVector;
-				bValidNTBs = false;
-			}
-
-			if (TmpNormal.IsNearlyZero() || TmpNormal.ContainsNaN())
-			{
-				TmpNormal = FVector3f::ZeroVector;
-				bValidNTBs = false;
-			}
-
-			Normal += TmpNormal;
-			Tangent += TmpTangent;
-			Binormal += TmpBinormal;
-		}
-		else
-		{
-			//This will force a recompute of the normals and tangents
-			Normal = FVector3f::ZeroVector;
-			Tangent = FVector3f::ZeroVector;
-			Binormal = FVector3f::ZeroVector;
-
-			// The polygon is degenerated
-			bValidNTBs = false;
-		}
-	}
-
-	PolygonNormals[PolygonID] = Normal.GetSafeNormal();
-	PolygonTangents[PolygonID] = Tangent.GetSafeNormal();
-	PolygonBinormals[PolygonID] = Binormal.GetSafeNormal();
-	PolygonCenters[PolygonID] = Center;
-
-	return bValidNTBs;
-}
-
-
-void FStaticMeshOperations::ComputePolygonTangentsAndNormals(FMeshDescription& MeshDescription, float ComparisonThreshold)
-{
-	PRAGMA_DISABLE_DEPRECATION_WARNINGS
-	TRACE_CPUPROFILER_EVENT_SCOPE(FStaticMeshOperations::ComputePolygonTangentsAndNormals_Selection);
-
-	FStaticMeshAttributes Attributes(MeshDescription);
-	Attributes.RegisterPolygonNormalAndTangentAttributes();
-
-	TArray<FPolygonID> PolygonIDs;
-	PolygonIDs.Reserve(MeshDescription.Polygons().Num());
-	for (const FPolygonID Polygon : MeshDescription.Polygons().GetElementIDs())
-	{
-		PolygonIDs.Add(Polygon);
-	}
-
-	// Split work in batch to reduce call overhead
-	const int32 BatchSize = 8 * 1024;
-	const int32 BatchCount = 1 + PolygonIDs.Num() / BatchSize;
-
-	ParallelFor(BatchCount,
-		[&PolygonIDs, &BatchSize, &ComparisonThreshold, &MeshDescription, &Attributes](int32 BatchIndex)
-		{
-			TVertexAttributesConstRef<FVector3f> VertexPositions = Attributes.GetVertexPositions();
-			TVertexInstanceAttributesConstRef<FVector2f> VertexUVs = Attributes.GetVertexInstanceUVs();
-			TPolygonAttributesRef<FVector3f> PolygonNormals = Attributes.GetPolygonNormals();
-			TPolygonAttributesRef<FVector3f> PolygonTangents = Attributes.GetPolygonTangents();
-			TPolygonAttributesRef<FVector3f> PolygonBinormals = Attributes.GetPolygonBinormals();
-			TPolygonAttributesRef<FVector3f> PolygonCenters = Attributes.GetPolygonCenters();
-
-			FVertexInstanceArray& VertexInstanceArray = MeshDescription.VertexInstances();
-			FVertexArray& VertexArray = MeshDescription.Vertices();
-			FPolygonArray& PolygonArray = MeshDescription.Polygons();
-
-			int32 Indice = BatchIndex * BatchSize;
-			int32 LastIndice = FMath::Min(Indice + BatchSize, PolygonIDs.Num());
-			for (; Indice < LastIndice; ++Indice)
-			{
-				const FPolygonID PolygonID = PolygonIDs[Indice];
-
-				if (!PolygonNormals[PolygonID].IsNearlyZero())
-				{
-					//By pass normal calculation if its already done
-					continue;
-				}
-
-				GetPolygonTangentsAndNormals(MeshDescription, PolygonID, ComparisonThreshold, VertexPositions, VertexUVs, PolygonNormals, PolygonTangents, PolygonBinormals, PolygonCenters);
-			}
-		}
-	);
-	PRAGMA_ENABLE_DEPRECATION_WARNINGS
-}
 
 
 static TTuple<FVector3f, FVector3f, FVector3f> GetTriangleTangentsAndNormalsWithUV(float ComparisonThreshold, TArrayView<const FVector3f> VertexPositions, TArrayView<const FVector2D> VertexUVs)
@@ -1036,7 +877,7 @@ void FStaticMeshOperations::ConvertFromRawMesh(const FRawMesh& SourceRawMesh, FM
 	DestinationMeshDescription.BuildIndexers();
 }
 
-void FStaticMeshOperations::AppendMeshDescriptions(const TArray<const FMeshDescription*>& SourceMeshes, FMeshDescription& TargetMesh, FAppendSettings& AppendSettings)
+void FStaticMeshOperations::AppendMeshDescriptions(const TArray<const FMeshDescription*>& SourceMeshes, FMeshDescription& TargetMesh, const FAppendSettings& AppendSettings)
 {
 	TRACE_CPUPROFILER_EVENT_SCOPE(FStaticMeshOperations::AppendMeshDescriptions);
 
@@ -1050,6 +891,26 @@ void FStaticMeshOperations::AppendMeshDescriptions(const TArray<const FMeshDescr
 	TVertexInstanceAttributesRef<FVector4f> TargetVertexInstanceColors = TargetAttributes.GetVertexInstanceColors();
 	TVertexInstanceAttributesRef<FVector2f> TargetVertexInstanceUVs = TargetAttributes.GetVertexInstanceUVs();
 
+	TPolygonAttributesRef<FName> TargetPolygonObjectNames;
+	if (AppendSettings.bMergeObjectName)
+	{
+		// If any of the source meshes have an object name attribute, but the target mesh doesn't, we register it here so that we can
+		// actually honor the bMergeObjectNames flag.
+		if (!TargetAttributes.GetPolygonObjectNames().IsValid())
+		{
+			for (const FMeshDescription* SourceMesh: SourceMeshes)
+			{
+				FStaticMeshConstAttributes SourceAttributes(*SourceMesh);
+				if (SourceAttributes.GetPolygonObjectNames().IsValid())
+				{
+					TargetAttributes.RegisterPolygonObjectNameAttribute();
+					break;
+				}
+			}
+		}
+		TargetPolygonObjectNames = TargetAttributes.GetPolygonObjectNames();
+	}
+
 	TargetMesh.SuspendVertexInstanceIndexing();
 	TargetMesh.SuspendEdgeIndexing();
 	TargetMesh.SuspendPolygonIndexing();
@@ -1138,6 +999,12 @@ void FStaticMeshOperations::AppendMeshDescriptions(const TArray<const FMeshDescr
 		TVertexInstanceAttributesConstRef<FVector4f> SourceVertexInstanceColors = SourceAttributes.GetVertexInstanceColors();
 		TVertexInstanceAttributesConstRef<FVector2f> SourceVertexInstanceUVs = SourceAttributes.GetVertexInstanceUVs();
 
+		TPolygonAttributesConstRef<FName> SourcePolygonObjectNames;
+		if (TargetPolygonObjectNames.IsValid())
+		{
+			SourcePolygonObjectNames = SourceAttributes.GetPolygonObjectNames();
+		}
+		
 		// Fill the UV arrays
 		const int32 NumUVChannel = FMath::Min(TargetMesh.GetNumUVElementChannels(), SourceMesh.GetNumUVElementChannels());
 		for (int32 UVLayerIndex = 0; UVLayerIndex < NumUVChannel; UVLayerIndex++)
@@ -1249,42 +1116,51 @@ void FStaticMeshOperations::AppendMeshDescriptions(const TArray<const FMeshDescr
 				FVertexInstanceID InstanceID = VertexInstanceIDPair.Value;
 
 				FVector3f& Normal = TargetVertexInstanceNormals[InstanceID];
-				Normal = (FVector3f)FVector(TransformInverseTransposeMatrix.TransformVector((FVector)Normal).GetSafeNormal());
+				Normal = static_cast<FVector3f>(FVector(TransformInverseTransposeMatrix.TransformVector(static_cast<FVector>(Normal)).GetSafeNormal()));
 
 				FVector3f& Tangent = TargetVertexInstanceTangents[InstanceID];
-				Tangent = (FVector3f)FVector(TransformInverseTransposeMatrix.TransformVector((FVector)Tangent).GetSafeNormal());
+				Tangent = static_cast<FVector3f>(FVector(TransformInverseTransposeMatrix.TransformVector(static_cast<FVector>(Tangent)).GetSafeNormal()));
 
 				TargetVertexInstanceBinormalSigns[InstanceID] *= BinormalSignsFactor;
 			}
 		}
 
-		// Triangles
-		for (const FTriangleID SourceTriangleID : SourceMesh.Triangles().GetElementIDs())
+		// Polygons
+		TArray<FVertexInstanceID> PolygonVertexInstanceIDs;
+		PolygonVertexInstanceIDs.Reserve(8);
+		
+		for (const FPolygonID SourcePolygonID: SourceMesh.Polygons().GetElementIDs())
 		{
-			TArrayView<const FVertexInstanceID> TriangleVertexInstanceIDs = SourceMesh.GetTriangleVertexInstances(SourceTriangleID);
-				
+			SourceMesh.GetPolygonVertexInstances(SourcePolygonID, PolygonVertexInstanceIDs);
+
 			//Find the polygonGroupID
-			FPolygonGroupID TargetPolygonGroupID = SinglePolygonGroup != INDEX_NONE ? SinglePolygonGroup : RemapPolygonGroup[SourceMesh.GetTrianglePolygonGroup(SourceTriangleID)];
+			FPolygonGroupID TargetPolygonGroupID = SinglePolygonGroup != INDEX_NONE ? SinglePolygonGroup : RemapPolygonGroup[SourceMesh.GetPolygonPolygonGroup(SourcePolygonID)];
 
-			TArray<FVertexInstanceID, TInlineAllocator<3>> VertexInstanceIDs;
-			VertexInstanceIDs.Reserve(3);
+			TArray<FVertexInstanceID, TInlineAllocator<4>> VertexInstanceIDs;
+			VertexInstanceIDs.Reserve(4);
 			if (bReverseCulling)
 			{
-				for (int32 ReverseVertexInstanceIdIndex = TriangleVertexInstanceIDs.Num()-1; ReverseVertexInstanceIdIndex >= 0; ReverseVertexInstanceIdIndex--)
+				for (int32 ReverseVertexInstanceIdIndex = PolygonVertexInstanceIDs.Num()-1; ReverseVertexInstanceIdIndex >= 0; ReverseVertexInstanceIdIndex--)
 				{
-					VertexInstanceIDs.Add(SourceToTargetVertexInstanceID[TriangleVertexInstanceIDs[ReverseVertexInstanceIdIndex]]);
+					VertexInstanceIDs.Add(SourceToTargetVertexInstanceID[PolygonVertexInstanceIDs[ReverseVertexInstanceIdIndex]]);
 				}
 			}
 			else
 			{
-				for (const FVertexInstanceID& VertexInstanceID : TriangleVertexInstanceIDs)
+				for (const FVertexInstanceID& VertexInstanceID : PolygonVertexInstanceIDs)
 				{
 					VertexInstanceIDs.Add(SourceToTargetVertexInstanceID[VertexInstanceID]);
 				}
 			}
-			
-			// Insert a triangle into the mesh
-			TargetMesh.CreateTriangle(TargetPolygonGroupID, VertexInstanceIDs);
+
+			FPolygonID TargetPolygonID = TargetMesh.CreatePolygon(TargetPolygonGroupID, VertexInstanceIDs);
+
+			// Copy the source object name, if requested. If the attribute doesn't exist on the source, just leave the name at default.
+			// SourcePolygonObjectNames can only be valid if TargetPolygonObjectNames, so we only check that here.
+			if (SourcePolygonObjectNames.IsValid())
+			{
+				TargetPolygonObjectNames.Set(TargetPolygonID, SourcePolygonObjectNames.Get(SourcePolygonID));
+			}
 		}
 	}
 
@@ -1295,7 +1171,7 @@ void FStaticMeshOperations::AppendMeshDescriptions(const TArray<const FMeshDescr
 	TargetMesh.ResumeUVIndexing();
 }
 
-void FStaticMeshOperations::AppendMeshDescription(const FMeshDescription& SourceMesh, FMeshDescription& TargetMesh, FAppendSettings& AppendSettings)
+void FStaticMeshOperations::AppendMeshDescription(const FMeshDescription& SourceMesh, FMeshDescription& TargetMesh, const FAppendSettings& AppendSettings)
 {
 	AppendMeshDescriptions({ &SourceMesh }, TargetMesh, AppendSettings);
 }
diff --git a/Engine/Source/Runtime/StaticMeshDescription/Public/StaticMeshAttributes.h b/Engine/Source/Runtime/StaticMeshDescription/Public/StaticMeshAttributes.h
--- a/Engine/Source/Runtime/StaticMeshDescription/Public/StaticMeshAttributes.h
+++ b/Engine/Source/Runtime/StaticMeshDescription/Public/StaticMeshAttributes.h
@@ -15,8 +15,6 @@ namespace MeshAttribute
 {
 	namespace Vertex
 	{
-		UE_DEPRECATED(4.26, "This attribute is no longer supported, please remove code pertaining to it.")
-		extern STATICMESHDESCRIPTION_API const FName CornerSharpness;	// deprecated
 	}
 
 	namespace VertexInstance
@@ -31,10 +29,6 @@ namespace MeshAttribute
 	namespace Edge
 	{
 		extern STATICMESHDESCRIPTION_API const FName IsHard;
-		UE_DEPRECATED(4.26, "This attribute is no longer supported, please remove code pertaining to it.")
-		extern STATICMESHDESCRIPTION_API const FName IsUVSeam;			// deprecated
-		UE_DEPRECATED(4.26, "This attribute is no longer supported, please remove code pertaining to it.")
-		extern STATICMESHDESCRIPTION_API const FName CreaseSharpness;	// deprecated
 	}
 
 	namespace Triangle
@@ -46,29 +40,16 @@ namespace MeshAttribute
 
 	namespace Polygon
 	{
-		UE_DEPRECATED(4.26, "This attribute is no longer supported, please remove code pertaining to it.")
-		extern STATICMESHDESCRIPTION_API const FName Normal;
-		UE_DEPRECATED(4.26, "This attribute is no longer supported, please remove code pertaining to it.")
-		extern STATICMESHDESCRIPTION_API const FName Tangent;
-		UE_DEPRECATED(4.26, "This attribute is no longer supported, please remove code pertaining to it.")
-		extern STATICMESHDESCRIPTION_API const FName Binormal;
-		UE_DEPRECATED(4.26, "This attribute is no longer supported, please remove code pertaining to it.")
-		extern STATICMESHDESCRIPTION_API const FName Center;
+		extern STATICMESHDESCRIPTION_API const FName ObjectName;
 	}
 
 	namespace PolygonGroup
 	{
 		extern STATICMESHDESCRIPTION_API const FName ImportedMaterialSlotName;
-		UE_DEPRECATED(4.26, "This attribute is no longer supported, please remove code pertaining to it.")
-		extern STATICMESHDESCRIPTION_API const FName EnableCollision;	// deprecated
-		UE_DEPRECATED(4.26, "This attribute is no longer supported, please remove code pertaining to it.")
-		extern STATICMESHDESCRIPTION_API const FName CastShadow;		// deprecated
 	}
 }
 
 
-PRAGMA_DISABLE_DEPRECATION_WARNINGS
-
 class FStaticMeshAttributes : public FMeshAttributes
 {
 public:
@@ -91,19 +72,18 @@ public:
                InAttributeName == MeshAttribute::Triangle::Normal ||
                InAttributeName == MeshAttribute::Triangle::Tangent ||
                InAttributeName == MeshAttribute::Triangle::Binormal ||
+               InAttributeName == MeshAttribute::Polygon::ObjectName ||
                InAttributeName == MeshAttribute::PolygonGroup::ImportedMaterialSlotName
 		;
 	}	
 
-	UE_DEPRECATED(4.26, "Please use RegisterTriangleNormalAndTangentAttributes() instead.")
-	STATICMESHDESCRIPTION_API void RegisterPolygonNormalAndTangentAttributes();
-
 	STATICMESHDESCRIPTION_API void RegisterTriangleNormalAndTangentAttributes();
 
-	UE_DEPRECATED(4.26, "This attribute is no longer supported, please remove code pertaining to it.")
-	TVertexAttributesRef<float> GetVertexCornerSharpnesses() { return MeshDescription.VertexAttributes().GetAttributesRef<float>(MeshAttribute::Vertex::CornerSharpness); }
-	UE_DEPRECATED(4.26, "This attribute is no longer supported, please remove code pertaining to it.")
-	TVertexAttributesConstRef<float> GetVertexCornerSharpnesses() const { return MeshDescription.VertexAttributes().GetAttributesRef<float>(MeshAttribute::Vertex::CornerSharpness); }
+	/** Registers a mesh attribute that is used to store the name of the object that contributed to each polygon. Used to partition the mesh
+	 *  into sub-objects, if needed.
+	 *  See GetPolygonObjectNames.
+	 */
+	STATICMESHDESCRIPTION_API void RegisterPolygonObjectNameAttribute();
 
 	TVertexInstanceAttributesRef<FVector2f> GetVertexInstanceUVs() { return MeshDescription.VertexInstanceAttributes().GetAttributesRef<FVector2f>(MeshAttribute::VertexInstance::TextureCoordinate); }
 	TVertexInstanceAttributesConstRef<FVector2f> GetVertexInstanceUVs() const { return MeshDescription.VertexInstanceAttributes().GetAttributesRef<FVector2f>(MeshAttribute::VertexInstance::TextureCoordinate); }
@@ -123,11 +103,6 @@ public:
 	TEdgeAttributesRef<bool> GetEdgeHardnesses() { return MeshDescription.EdgeAttributes().GetAttributesRef<bool>(MeshAttribute::Edge::IsHard); }
 	TEdgeAttributesConstRef<bool> GetEdgeHardnesses() const { return MeshDescription.EdgeAttributes().GetAttributesRef<bool>(MeshAttribute::Edge::IsHard); }
 
-	UE_DEPRECATED(4.26, "This attribute is no longer supported, please remove code pertaining to it.")
-	TEdgeAttributesRef<float> GetEdgeCreaseSharpnesses() { return MeshDescription.EdgeAttributes().GetAttributesRef<float>(MeshAttribute::Edge::CreaseSharpness); }
-	UE_DEPRECATED(4.26, "This attribute is no longer supported, please remove code pertaining to it.")
-	TEdgeAttributesConstRef<float> GetEdgeCreaseSharpnesses() const { return MeshDescription.EdgeAttributes().GetAttributesRef<float>(MeshAttribute::Edge::CreaseSharpness); }
-
 	TTriangleAttributesRef<FVector3f> GetTriangleNormals() { return MeshDescription.TriangleAttributes().GetAttributesRef<FVector3f>(MeshAttribute::Triangle::Normal); }
 	TTriangleAttributesConstRef<FVector3f> GetTriangleNormals() const { return MeshDescription.TriangleAttributes().GetAttributesRef<FVector3f>(MeshAttribute::Triangle::Normal); }
 
@@ -137,26 +112,10 @@ public:
 	TTriangleAttributesRef<FVector3f> GetTriangleBinormals() { return MeshDescription.TriangleAttributes().GetAttributesRef<FVector3f>(MeshAttribute::Triangle::Binormal); }
 	TTriangleAttributesConstRef<FVector3f> GetTriangleBinormals() const { return MeshDescription.TriangleAttributes().GetAttributesRef<FVector3f>(MeshAttribute::Triangle::Binormal); }
 
-	UE_DEPRECATED(4.26, "This attribute is no longer supported, please remove code pertaining to it.")
-	TPolygonAttributesRef<FVector3f> GetPolygonNormals() { return MeshDescription.PolygonAttributes().GetAttributesRef<FVector3f>(MeshAttribute::Polygon::Normal); }
-	UE_DEPRECATED(4.26, "This attribute is no longer supported, please remove code pertaining to it.")
-	TPolygonAttributesConstRef<FVector3f> GetPolygonNormals() const { return MeshDescription.PolygonAttributes().GetAttributesRef<FVector3f>(MeshAttribute::Polygon::Normal); }
-
-	UE_DEPRECATED(4.26, "This attribute is no longer supported, please remove code pertaining to it.")
-	TPolygonAttributesRef<FVector3f> GetPolygonTangents() { return MeshDescription.PolygonAttributes().GetAttributesRef<FVector3f>(MeshAttribute::Polygon::Tangent); }
-	UE_DEPRECATED(4.26, "This attribute is no longer supported, please remove code pertaining to it.")
-	TPolygonAttributesConstRef<FVector3f> GetPolygonTangents() const { return MeshDescription.PolygonAttributes().GetAttributesRef<FVector3f>(MeshAttribute::Polygon::Tangent); }
-
-	UE_DEPRECATED(4.26, "This attribute is no longer supported, please remove code pertaining to it.")
-	TPolygonAttributesRef<FVector3f> GetPolygonBinormals() { return MeshDescription.PolygonAttributes().GetAttributesRef<FVector3f>(MeshAttribute::Polygon::Binormal); }
-	UE_DEPRECATED(4.26, "This attribute is no longer supported, please remove code pertaining to it.")
-	TPolygonAttributesConstRef<FVector3f> GetPolygonBinormals() const { return MeshDescription.PolygonAttributes().GetAttributesRef<FVector3f>(MeshAttribute::Polygon::Binormal); }
-
-	UE_DEPRECATED(4.26, "This attribute is no longer supported, please remove code pertaining to it.")
-	TPolygonAttributesRef<FVector3f> GetPolygonCenters() { return MeshDescription.PolygonAttributes().GetAttributesRef<FVector3f>(MeshAttribute::Polygon::Center); }
-	UE_DEPRECATED(4.26, "This attribute is no longer supported, please remove code pertaining to it.")
-	TPolygonAttributesConstRef<FVector3f> GetPolygonCenters() const { return MeshDescription.PolygonAttributes().GetAttributesRef<FVector3f>(MeshAttribute::Polygon::Center); }
-
+	/** This attribute stores the name of the object that this part of the mesh was generated from, e.g. from a sub-object during import or during mesh merging. */
+	TPolygonAttributesRef<FName> GetPolygonObjectNames() { return MeshDescription.PolygonAttributes().GetAttributesRef<FName>(MeshAttribute::Polygon::ObjectName); }
+	TPolygonAttributesConstRef<FName> GetPolygonObjectNames() const { return MeshDescription.PolygonAttributes().GetAttributesRef<FName>(MeshAttribute::Polygon::ObjectName); }
+	
 	TPolygonGroupAttributesRef<FName> GetPolygonGroupMaterialSlotNames() { return MeshDescription.PolygonGroupAttributes().GetAttributesRef<FName>(MeshAttribute::PolygonGroup::ImportedMaterialSlotName); }
 	TPolygonGroupAttributesConstRef<FName> GetPolygonGroupMaterialSlotNames() const { return MeshDescription.PolygonGroupAttributes().GetAttributesRef<FName>(MeshAttribute::PolygonGroup::ImportedMaterialSlotName); }
 };
@@ -170,28 +129,15 @@ public:
 		: FMeshConstAttributes(InMeshDescription)
 	{}
 
-	UE_DEPRECATED(4.26, "This attribute is no longer supported, please remove code pertaining to it.")
-	TVertexAttributesConstRef<float> GetVertexCornerSharpnesses() const { return MeshDescription.VertexAttributes().GetAttributesRef<float>(MeshAttribute::Vertex::CornerSharpness); }
 	TVertexInstanceAttributesConstRef<FVector2f> GetVertexInstanceUVs() const { return MeshDescription.VertexInstanceAttributes().GetAttributesRef<FVector2f>(MeshAttribute::VertexInstance::TextureCoordinate); }
 	TVertexInstanceAttributesConstRef<FVector3f> GetVertexInstanceNormals() const { return MeshDescription.VertexInstanceAttributes().GetAttributesRef<FVector3f>(MeshAttribute::VertexInstance::Normal); }
 	TVertexInstanceAttributesConstRef<FVector3f> GetVertexInstanceTangents() const { return MeshDescription.VertexInstanceAttributes().GetAttributesRef<FVector3f>(MeshAttribute::VertexInstance::Tangent); }
 	TVertexInstanceAttributesConstRef<float> GetVertexInstanceBinormalSigns() const { return MeshDescription.VertexInstanceAttributes().GetAttributesRef<float>(MeshAttribute::VertexInstance::BinormalSign); }
 	TVertexInstanceAttributesConstRef<FVector4f> GetVertexInstanceColors() const { return MeshDescription.VertexInstanceAttributes().GetAttributesRef<FVector4f>(MeshAttribute::VertexInstance::Color); }
 	TEdgeAttributesConstRef<bool> GetEdgeHardnesses() const { return MeshDescription.EdgeAttributes().GetAttributesRef<bool>(MeshAttribute::Edge::IsHard); }
-	UE_DEPRECATED(4.26, "This attribute is no longer supported, please remove code pertaining to it.")
-	TEdgeAttributesConstRef<float> GetEdgeCreaseSharpnesses() const { return MeshDescription.EdgeAttributes().GetAttributesRef<float>(MeshAttribute::Edge::CreaseSharpness); }
 	TTriangleAttributesConstRef<FVector3f> GetTriangleNormals() const { return MeshDescription.TriangleAttributes().GetAttributesRef<FVector3f>(MeshAttribute::Triangle::Normal); }
 	TTriangleAttributesConstRef<FVector3f> GetTriangleTangents() const { return MeshDescription.TriangleAttributes().GetAttributesRef<FVector3f>(MeshAttribute::Triangle::Tangent); }
 	TTriangleAttributesConstRef<FVector3f> GetTriangleBinormals() const { return MeshDescription.TriangleAttributes().GetAttributesRef<FVector3f>(MeshAttribute::Triangle::Binormal); }
-	UE_DEPRECATED(4.26, "This attribute is no longer supported, please remove code pertaining to it.")
-	TPolygonAttributesConstRef<FVector3f> GetPolygonNormals() const { return MeshDescription.PolygonAttributes().GetAttributesRef<FVector3f>(MeshAttribute::Polygon::Normal); }
-	UE_DEPRECATED(4.26, "This attribute is no longer supported, please remove code pertaining to it.")
-	TPolygonAttributesConstRef<FVector3f> GetPolygonTangents() const { return MeshDescription.PolygonAttributes().GetAttributesRef<FVector3f>(MeshAttribute::Polygon::Tangent); }
-	UE_DEPRECATED(4.26, "This attribute is no longer supported, please remove code pertaining to it.")
-	TPolygonAttributesConstRef<FVector3f> GetPolygonBinormals() const { return MeshDescription.PolygonAttributes().GetAttributesRef<FVector3f>(MeshAttribute::Polygon::Binormal); }
-	UE_DEPRECATED(4.26, "This attribute is no longer supported, please remove code pertaining to it.")
-	TPolygonAttributesConstRef<FVector3f> GetPolygonCenters() const { return MeshDescription.PolygonAttributes().GetAttributesRef<FVector3f>(MeshAttribute::Polygon::Center); }
+	TPolygonAttributesConstRef<FName> GetPolygonObjectNames() const { return MeshDescription.PolygonAttributes().GetAttributesRef<FName>(MeshAttribute::Polygon::ObjectName); }
 	TPolygonGroupAttributesConstRef<FName> GetPolygonGroupMaterialSlotNames() const { return MeshDescription.PolygonGroupAttributes().GetAttributesRef<FName>(MeshAttribute::PolygonGroup::ImportedMaterialSlotName); }
 };
-
-PRAGMA_ENABLE_DEPRECATION_WARNINGS
diff --git a/Engine/Source/Runtime/StaticMeshDescription/Public/StaticMeshDescription.h b/Engine/Source/Runtime/StaticMeshDescription/Public/StaticMeshDescription.h
--- a/Engine/Source/Runtime/StaticMeshDescription/Public/StaticMeshDescription.h
+++ b/Engine/Source/Runtime/StaticMeshDescription/Public/StaticMeshDescription.h
@@ -20,8 +20,6 @@ class UMaterial;
 class UObject;
 struct FFrame;
 
-PRAGMA_DISABLE_DEPRECATION_WARNINGS
-
 /**
 * A wrapper for MeshDescription, customized for static meshes
 */
@@ -64,11 +62,6 @@ public:
 
 public:
 
-	UE_DEPRECATED(4.25, "This attribute is no longer supported, please remove code pertaining to it.")
-	TVertexAttributesRef<float> GetVertexCornerSharpnesses() { return GetRequiredAttributes().GetVertexCornerSharpnesses(); }
-	UE_DEPRECATED(4.25, "This attribute is no longer supported, please remove code pertaining to it.")
-	TVertexAttributesConstRef<float> GetVertexCornerSharpnesses() const { return GetRequiredAttributes().GetVertexCornerSharpnesses(); }
-
 	TVertexInstanceAttributesRef<FVector2f> GetVertexInstanceUVs() { return GetRequiredAttributes().GetVertexInstanceUVs(); }
 	TVertexInstanceAttributesConstRef<FVector2f> GetVertexInstanceUVs() const { return GetRequiredAttributes().GetVertexInstanceUVs(); }
 
@@ -87,13 +80,6 @@ public:
 	TEdgeAttributesRef<bool> GetEdgeHardnesses() { return GetRequiredAttributes().GetEdgeHardnesses(); }
 	TEdgeAttributesConstRef<bool> GetEdgeHardnesses() const { return GetRequiredAttributes().GetEdgeHardnesses(); }
 
-	UE_DEPRECATED(4.25, "This attribute is no longer supported, please remove code pertaining to it.")
-	TEdgeAttributesRef<float> GetEdgeCreaseSharpnesses() { return GetRequiredAttributes().GetEdgeCreaseSharpnesses(); }
-	UE_DEPRECATED(4.25, "This attribute is no longer supported, please remove code pertaining to it.")
-	TEdgeAttributesConstRef<float> GetEdgeCreaseSharpnesses() const { return GetRequiredAttributes().GetEdgeCreaseSharpnesses(); }
-
 	TPolygonGroupAttributesRef<FName> GetPolygonGroupMaterialSlotNames() { return GetRequiredAttributes().GetPolygonGroupMaterialSlotNames(); }
 	TPolygonGroupAttributesConstRef<FName> GetPolygonGroupMaterialSlotNames() const { return GetRequiredAttributes().GetPolygonGroupMaterialSlotNames(); }
 };
-
-PRAGMA_ENABLE_DEPRECATION_WARNINGS
diff --git a/Engine/Source/Runtime/StaticMeshDescription/Public/StaticMeshOperations.h b/Engine/Source/Runtime/StaticMeshDescription/Public/StaticMeshOperations.h
--- a/Engine/Source/Runtime/StaticMeshDescription/Public/StaticMeshOperations.h
+++ b/Engine/Source/Runtime/StaticMeshDescription/Public/StaticMeshOperations.h
@@ -48,28 +48,19 @@ class FStaticMeshOperations
 public:
 	struct FAppendSettings
 	{
-		FAppendSettings()
-			: bMergeVertexColor(true)
-			, bMergeUVChannels{ true }
-			, MergedAssetPivot(0.0f, 0.0f, 0.0f)
-		{}
-
 		enum
 		{
 			MAX_NUM_UV_CHANNELS = 8,
 		};
 
 		FAppendPolygonGroupsDelegate PolygonGroupsDelegate;
-		bool bMergeVertexColor;
-		bool bMergeUVChannels[MAX_NUM_UV_CHANNELS];
-		FVector MergedAssetPivot;
+		bool bMergeVertexColor = true;
+		bool bMergeUVChannels[MAX_NUM_UV_CHANNELS] = {true};
+		bool bMergeObjectName = false;
+		FVector MergedAssetPivot = FVector::ZeroVector;
 		TOptional<FTransform> MeshTransform; // Apply a transformation on source mesh (see MeshTransform)
 	};
 
-	/** Set the polygon tangent, normal, binormal and polygonCenter for all polygons in the mesh description. */
-	UE_DEPRECATED(4.26, "Please use ComputeTriangleTangentsAndNormals() instead.")
-	static STATICMESHDESCRIPTION_API void ComputePolygonTangentsAndNormals(FMeshDescription& MeshDescription, float ComparisonThreshold = 0.0f);
-
 	/** Set the triangle tangent, normal, binormal and triangleCenter for all triangles in the mesh description. */
 	static STATICMESHDESCRIPTION_API void ComputeTriangleTangentsAndNormals(FMeshDescription& MeshDescription, float ComparisonThreshold = 0.0f, const TCHAR* DebugName = nullptr);
 
@@ -96,8 +87,8 @@ public:
 	/** Convert old FRawMesh format to MeshDescription. */
 	static STATICMESHDESCRIPTION_API void ConvertFromRawMesh(const FRawMesh& SourceRawMesh, FMeshDescription& DestinationMeshDescription, const TMap<int32, FName>& MaterialMap, bool bSkipNormalsAndTangents = false, const TCHAR* DebugName = nullptr);
 
-	static STATICMESHDESCRIPTION_API void AppendMeshDescription(const FMeshDescription& SourceMesh, FMeshDescription& TargetMesh, FAppendSettings& AppendSettings);
-	static STATICMESHDESCRIPTION_API void AppendMeshDescriptions(const TArray<const FMeshDescription*>& SourceMeshes, FMeshDescription& TargetMesh, FAppendSettings& AppendSettings);
+	static STATICMESHDESCRIPTION_API void AppendMeshDescription(const FMeshDescription& SourceMesh, FMeshDescription& TargetMesh, const FAppendSettings& AppendSettings);
+	static STATICMESHDESCRIPTION_API void AppendMeshDescriptions(const TArray<const FMeshDescription*>& SourceMeshes, FMeshDescription& TargetMesh, const FAppendSettings& AppendSettings);
 
 	/** Check the mesh description's vertex instance normals and tangents to see if any of the values are invalid (e.g. nearly zero or they contain a NaN).
 	 *  \note This does not check if the vectors are normalized (which they should for correct normals and tangents).
