Mesh Description: Add a utility function to get the area of a mesh.

#rnx

diff --git a/Engine/Source/Runtime/StaticMeshDescription/Private/StaticMeshOperations.cpp b/Engine/Source/Runtime/StaticMeshDescription/Private/StaticMeshOperations.cpp
--- a/Engine/Source/Runtime/StaticMeshDescription/Private/StaticMeshOperations.cpp
+++ b/Engine/Source/Runtime/StaticMeshDescription/Private/StaticMeshOperations.cpp
@@ -3075,4 +3075,21 @@ bool FStaticMeshOperations::ValidateAndFixData(FMeshDescription& MeshDescription
 	return !bHasInvalidPositions && !bHasInvalidTangentSpaces && !bHasInvalidUVs && !bHasInvalidVertexColors;
 }
 
+float FStaticMeshOperations::ComputeMeshArea(const FMeshDescription& InMeshDescription)
+{
+	float AreaSum = 0.0f;
+
+	for (FTriangleID TriangleID: InMeshDescription.Triangles().GetElementIDs())
+	{
+		TArrayView<const FVertexID> TriangleVertices = InMeshDescription.GetTriangleVertices(TriangleID);
+		const FVector3f& PointA = InMeshDescription.GetVertexPosition(TriangleVertices[0]);
+		const FVector3f& PointB = InMeshDescription.GetVertexPosition(TriangleVertices[1]);
+		const FVector3f& PointC = InMeshDescription.GetVertexPosition(TriangleVertices[2]);
+
+		AreaSum += TriangleUtilities::ComputeTriangleArea(PointA, PointB, PointC);
+	}
+
+	return AreaSum;
+}
+
 #undef LOCTEXT_NAMESPACE
diff --git a/Engine/Source/Runtime/StaticMeshDescription/Public/StaticMeshOperations.h b/Engine/Source/Runtime/StaticMeshDescription/Public/StaticMeshOperations.h
--- a/Engine/Source/Runtime/StaticMeshDescription/Public/StaticMeshOperations.h
+++ b/Engine/Source/Runtime/StaticMeshDescription/Public/StaticMeshOperations.h
@@ -235,4 +235,7 @@ public:
 
 	/** Verify the mesh data does not contain any NAN or INF float value, if such a case happen the value are set to zero or identity for matrix or quat. */
 	static STATICMESHDESCRIPTION_API bool ValidateAndFixData(FMeshDescription& MeshDescription, const FString& DebugName);
+
+	/** Compute the surface area of the mesh by iterating over each triangle. */
+	static STATICMESHDESCRIPTION_API float ComputeMeshArea(const FMeshDescription& InMeshDescription);
 };
