UE: Add some debugging logs to catch invalid mesh data when building

#Synced-CL 50742541

diff --git a/Engine/Source/Runtime/StaticMeshDescription/Private/StaticMeshOperations.cpp b/Engine/Source/Runtime/StaticMeshDescription/Private/StaticMeshOperations.cpp
--- a/Engine/Source/Runtime/StaticMeshDescription/Private/StaticMeshOperations.cpp
+++ b/Engine/Source/Runtime/StaticMeshDescription/Private/StaticMeshOperations.cpp
@@ -181,6 +181,23 @@ void FStaticMeshOperations::ComputeTriangleTangentsAndNormals(FMeshDescription&
 					continue;
 				}
 
+				if (!TriangleVertexIDs.IsValidIndex(TriIndex + 2))
+				{
+					UE_LOG(LogStaticMeshOperations, Warning, TEXT("Static Mesh %s has invalid TriangleVertexIDs length! TriIndex %d. Skipping"),
+						DebugName ? DebugName : TEXT("UnknownMesh"), TriIndex + 2);
+					continue;
+				}
+
+				if (!VertexPositions.IsValidIndex(TriangleVertexIDs[TriIndex]) ||
+					!VertexPositions.IsValidIndex(TriangleVertexIDs[TriIndex + 1]) ||
+					!VertexPositions.IsValidIndex(TriangleVertexIDs[TriIndex + 2])
+					)
+				{
+					UE_LOG(LogStaticMeshOperations, Warning, TEXT("Static Mesh %s has invalid vertex ids! StartIndex %d. VertId0=%d, VertId1=%d, VertId2=%d. Skipping"),
+						DebugName ? DebugName : TEXT("UnknownMesh"), StartIndex, TriangleVertexIDs[TriIndex].GetValue(), TriangleVertexIDs[TriIndex + 1].GetValue(), TriangleVertexIDs[TriIndex + 2].GetValue());
+					continue;
+				}
+
 				FVector3f TriangleVertexPositions[3] =
 				{
 					VertexPositions[TriangleVertexIDs[TriIndex]],
@@ -202,6 +219,24 @@ void FStaticMeshOperations::ComputeTriangleTangentsAndNormals(FMeshDescription&
 				TTuple<FVector3f, FVector3f, FVector3f> Result;
 				if (!VertexUVs.IsEmpty())
 				{
+					if (!TriangleVertexInstanceIDs.IsValidIndex(TriIndex + 2))
+					{
+						UE_LOG(LogStaticMeshOperations, Warning, TEXT("Static Mesh %s has invalid TriangleVertexInstanceIDs length! TriIndex %d. Skipping"),
+							DebugName ? DebugName : TEXT("UnknownMesh"), TriIndex + 2);
+						continue;
+					}
+
+					if (!VertexUVs.IsValidIndex(TriangleVertexInstanceIDs[TriIndex]) ||
+						!VertexUVs.IsValidIndex(TriangleVertexInstanceIDs[TriIndex + 1]) ||
+						!VertexUVs.IsValidIndex(TriangleVertexInstanceIDs[TriIndex + 2])
+						)
+					{
+						// Invalid input data
+						UE_LOG(LogStaticMeshOperations, Warning, TEXT("Static Mesh %s has invalid vertex instance ids! Triangle index %d. InstanceId0=%d, InstanceId1=%d, InstanceId2=%d. Skipping"),
+							DebugName ? DebugName : TEXT("UnknownMesh"), StartIndex, TriangleVertexInstanceIDs[TriIndex].GetValue(), TriangleVertexInstanceIDs[TriIndex + 1].GetValue(), TriangleVertexInstanceIDs[TriIndex + 2].GetValue());
+						continue;
+					}
+
 					FVector2D TriangleUVs[3] =
 					{
 						FVector2D(VertexUVs[TriangleVertexInstanceIDs[TriIndex]]),
