UE: Fill out TriangleNormals/Tangents/Binormals for invalid verts when building mesh data instead of leaving them undefined

diff --git a/Engine/Source/Runtime/StaticMeshDescription/Private/StaticMeshOperations.cpp b/Engine/Source/Runtime/StaticMeshDescription/Private/StaticMeshOperations.cpp
--- a/Engine/Source/Runtime/StaticMeshDescription/Private/StaticMeshOperations.cpp
+++ b/Engine/Source/Runtime/StaticMeshDescription/Private/StaticMeshOperations.cpp
@@ -185,6 +185,9 @@ void FStaticMeshOperations::ComputeTriangleTangentsAndNormals(FMeshDescription&
 				{
 					UE_LOG(LogStaticMeshOperations, Warning, TEXT("Static Mesh %s has invalid TriangleVertexIDs length! TriIndex %d. Skipping"),
 						DebugName ? DebugName : TEXT("UnknownMesh"), TriIndex + 2);
+					TriangleNormals[StartIndex] = FVector3f(1, 0, 0);
+					TriangleTangents[StartIndex] = FVector3f(0, 1, 0);
+					TriangleBinormals[StartIndex] = FVector3f(0, 0, 1);
 					continue;
 				}
 
@@ -195,6 +198,9 @@ void FStaticMeshOperations::ComputeTriangleTangentsAndNormals(FMeshDescription&
 				{
 					UE_LOG(LogStaticMeshOperations, Warning, TEXT("Static Mesh %s has invalid vertex ids! StartIndex %d. VertId0=%d, VertId1=%d, VertId2=%d. Skipping"),
 						DebugName ? DebugName : TEXT("UnknownMesh"), StartIndex, TriangleVertexIDs[TriIndex].GetValue(), TriangleVertexIDs[TriIndex + 1].GetValue(), TriangleVertexIDs[TriIndex + 2].GetValue());
+					TriangleNormals[StartIndex] = FVector3f(1, 0, 0);
+					TriangleTangents[StartIndex] = FVector3f(0, 1, 0);
+					TriangleBinormals[StartIndex] = FVector3f(0, 0, 1);
 					continue;
 				}
 
@@ -223,6 +229,9 @@ void FStaticMeshOperations::ComputeTriangleTangentsAndNormals(FMeshDescription&
 					{
 						UE_LOG(LogStaticMeshOperations, Warning, TEXT("Static Mesh %s has invalid TriangleVertexInstanceIDs length! TriIndex %d. Skipping"),
 							DebugName ? DebugName : TEXT("UnknownMesh"), TriIndex + 2);
+						TriangleNormals[StartIndex] = FVector3f(1, 0, 0);
+						TriangleTangents[StartIndex] = FVector3f(0, 1, 0);
+						TriangleBinormals[StartIndex] = FVector3f(0, 0, 1);
 						continue;
 					}
 
@@ -231,9 +240,11 @@ void FStaticMeshOperations::ComputeTriangleTangentsAndNormals(FMeshDescription&
 						!VertexUVs.IsValidIndex(TriangleVertexInstanceIDs[TriIndex + 2])
 						)
 					{
-						// Invalid input data
 						UE_LOG(LogStaticMeshOperations, Warning, TEXT("Static Mesh %s has invalid vertex instance ids! Triangle index %d. InstanceId0=%d, InstanceId1=%d, InstanceId2=%d. Skipping"),
 							DebugName ? DebugName : TEXT("UnknownMesh"), StartIndex, TriangleVertexInstanceIDs[TriIndex].GetValue(), TriangleVertexInstanceIDs[TriIndex + 1].GetValue(), TriangleVertexInstanceIDs[TriIndex + 2].GetValue());
+						TriangleNormals[StartIndex] = FVector3f(1, 0, 0);
+						TriangleTangents[StartIndex] = FVector3f(0, 1, 0);
+						TriangleBinormals[StartIndex] = FVector3f(0, 0, 1);
 						continue;
 					}
 
